HierarchyTable Core

diff --git a/Engine/Plugins/Experimental/Animation/HierarchyTable/HierarchyTable.uplugin b/Engine/Plugins/Experimental/Animation/HierarchyTable/HierarchyTable.uplugin
new file mode 100644
--- /dev/null
+++ b/Engine/Plugins/Experimental/Animation/HierarchyTable/HierarchyTable.uplugin
@@ -0,0 +1,36 @@
+{
+	"FileVersion": 3,
+	"Version": 1,
+	"VersionName": "0.1",
+	"FriendlyName": "Hierarchy Table",
+	"Description": "",
+	"Category": "Animation",
+	"CreatedBy": "Epic Games, Inc.",
+	"CreatedByURL": "https://epicgames.com",
+	"DocsURL": "",
+	"MarketplaceURL": "",
+	"SupportURL": "",
+	"EnabledByDefault": false,
+	"CanContainContent": true,
+	"IsBetaVersion": false,
+	"IsExperimentalVersion": true,
+	"Installed": false,
+	"Modules": [
+		{
+			"Name": "HierarchyTableRuntime",
+			"Type": "Runtime",
+			"LoadingPhase": "Default"
+		},
+		{
+			"Name": "HierarchyTableEditor",
+			"Type": "UncookedOnly",
+			"LoadingPhase": "Default"
+		}
+	],
+	"Plugins": [
+		{
+			"Name": "TypedElementDataStorageEditorUI",
+			"Enabled": true
+		}
+	]
+}
diff --git a/Engine/Plugins/Experimental/Animation/HierarchyTable/Source/Editor/HierarchyTableEditor.Build.cs b/Engine/Plugins/Experimental/Animation/HierarchyTable/Source/Editor/HierarchyTableEditor.Build.cs
new file mode 100644
--- /dev/null
+++ b/Engine/Plugins/Experimental/Animation/HierarchyTable/Source/Editor/HierarchyTableEditor.Build.cs
@@ -0,0 +1,28 @@
+// Copyright Epic Games, Inc. All Rights Reserved.
+
+using EpicGames.Core;
+using UnrealBuildTool;
+
+public class HierarchyTableEditor : ModuleRules
+{
+	public HierarchyTableEditor(ReadOnlyTargetRules Target) : base(Target)
+	{
+		PrivateDependencyModuleNames.AddAll(
+			"AnimGraph",
+			"AssetDefinition",
+			"Core",
+			"CoreUObject",
+			"Engine", 
+			"InputCore",
+			"Slate",
+			"SlateCore",
+			"UnrealEd",
+			"HierarchyTableRuntime",
+			"TedsOutliner",
+			"TypedElementFramework",
+			"SceneOutliner",
+			"ToolMenus",
+			"Persona"
+		);
+	}
+}
diff --git a/Engine/Plugins/Experimental/Animation/HierarchyTable/Source/Editor/Private/HierachyTableEditorModule.cpp b/Engine/Plugins/Experimental/Animation/HierarchyTable/Source/Editor/Private/HierachyTableEditorModule.cpp
new file mode 100644
--- /dev/null
+++ b/Engine/Plugins/Experimental/Animation/HierarchyTable/Source/Editor/Private/HierachyTableEditorModule.cpp
@@ -0,0 +1,14 @@
+// Copyright Epic Games, Inc. All Rights Reserved.
+
+#include "HierarchyTableEditorModule.h"
+#include "Modules/ModuleManager.h"
+
+void FHierarchyTableEditorModule::StartupModule()
+{
+}
+
+void FHierarchyTableEditorModule::ShutdownModule()
+{
+}
+
+IMPLEMENT_MODULE(FHierarchyTableEditorModule, HierarchyTableEditor)
diff --git a/Engine/Plugins/Experimental/Animation/HierarchyTable/Source/Editor/Private/HierarchyTable/Factories/OverrideFactory.cpp b/Engine/Plugins/Experimental/Animation/HierarchyTable/Source/Editor/Private/HierarchyTable/Factories/OverrideFactory.cpp
new file mode 100644
--- /dev/null
+++ b/Engine/Plugins/Experimental/Animation/HierarchyTable/Source/Editor/Private/HierarchyTable/Factories/OverrideFactory.cpp
@@ -0,0 +1,13 @@
+// Copyright Epic Games, Inc. All Rights Reserved.
+
+#include "HierarchyTable/Factories/OverrideFactory.h"
+#include "HierarchyTable/Columns/OverrideColumn.h"
+#include "HierarchyTable/Widgets/OverrideWidgetConstructor.h"
+#include "Elements/Interfaces/TypedElementDataStorageUiInterface.h"
+
+void UHierarchyTableOverrideFactory::RegisterWidgetConstructors(ITypedElementDataStorageInterface& DataStorage, ITypedElementDataStorageUiInterface& DataStorageUi) const
+{
+	using namespace TypedElementDataStorage;
+
+	DataStorageUi.RegisterWidgetFactory<FTypedElementWidgetConstructor_Override>(FName(TEXT("General.Cell")), FColumn<FTypedElementOverrideColumn>());
+}
\ No newline at end of file
diff --git a/Engine/Plugins/Experimental/Animation/HierarchyTable/Source/Editor/Private/HierarchyTable/Factories/OverrideFactory.h b/Engine/Plugins/Experimental/Animation/HierarchyTable/Source/Editor/Private/HierarchyTable/Factories/OverrideFactory.h
new file mode 100644
--- /dev/null
+++ b/Engine/Plugins/Experimental/Animation/HierarchyTable/Source/Editor/Private/HierarchyTable/Factories/OverrideFactory.h
@@ -0,0 +1,15 @@
+// Copyright Epic Games, Inc. All Rights Reserved.
+
+#pragma once
+
+#include "Elements/Interfaces/TypedElementDataStorageFactory.h"
+#include "OverrideFactory.generated.h"
+
+UCLASS()
+class UHierarchyTableOverrideFactory : public UTypedElementDataStorageFactory
+{
+	GENERATED_BODY()
+
+public:
+	void RegisterWidgetConstructors(ITypedElementDataStorageInterface& DataStorage, ITypedElementDataStorageUiInterface& DataStorageUi) const override;
+};
\ No newline at end of file
diff --git a/Engine/Plugins/Experimental/Animation/HierarchyTable/Source/Editor/Private/HierarchyTable/Widgets/OverrideWidgetConstructor.cpp b/Engine/Plugins/Experimental/Animation/HierarchyTable/Source/Editor/Private/HierarchyTable/Widgets/OverrideWidgetConstructor.cpp
new file mode 100644
--- /dev/null
+++ b/Engine/Plugins/Experimental/Animation/HierarchyTable/Source/Editor/Private/HierarchyTable/Widgets/OverrideWidgetConstructor.cpp
@@ -0,0 +1,107 @@
+// Copyright Epic Games, Inc. All Rights Reserved.
+
+#include "HierarchyTable/Widgets/OverrideWidgetConstructor.h"
+#include "Elements/Columns/TypedElementMiscColumns.h"
+#include "HierarchyTable/Columns/OverrideColumn.h"
+#include "Widgets/Input/SButton.h"
+#include "Widgets/Images/SImage.h"
+#include "Widgets/Layout/SBox.h"
+#include "Widgets/SBoxPanel.h"
+#include "HierarchyTable.h"
+
+FTypedElementWidgetConstructor_Override::FTypedElementWidgetConstructor_Override()
+	: Super(StaticStruct())
+{ 
+}
+
+TSharedPtr<SWidget> FTypedElementWidgetConstructor_Override::CreateWidget(
+	const TypedElementDataStorage::FMetaDataView& Arguments)
+{
+	return SNew(SBox)
+		.HAlign(HAlign_Left)
+		.VAlign(VAlign_Center);
+}
+
+bool FTypedElementWidgetConstructor_Override::FinalizeWidget(
+	ITypedElementDataStorageInterface* DataStorage,
+	ITypedElementDataStorageUiInterface* DataStorageUi,
+	TypedElementDataStorage::RowHandle Row,
+	const TSharedPtr<SWidget>& Widget)
+{
+	if (!Widget)
+	{
+		return true;
+	}
+
+	checkf(Widget->GetType() == SBox::StaticWidgetClass().GetWidgetType(),
+		TEXT("Stored widget with FTypedElementFloatWidgetConstructor doesn't match type %s, but was a %s."),
+		*(SBox::StaticWidgetClass().GetWidgetType().ToString()),
+		*(Widget->GetTypeAsString()));
+
+	TSharedPtr<SBox> WidgetInstance = StaticCastSharedPtr<SBox>(Widget);
+
+	// Row is not actually the row we want, its contained inside of a row reference, I don't know why things are done this way.
+	const TypedElementDataStorage::RowHandle TargetRow = DataStorage->GetColumn<FTypedElementRowReferenceColumn>(Row)->Row;
+	FTypedElementOverrideColumn* OverrideColumn = DataStorage->GetColumn<FTypedElementOverrideColumn>(TargetRow);
+
+	if (OverrideColumn == nullptr)
+	{
+		check(false);
+		WidgetInstance->SetContent(SNullWidget::NullWidget);
+
+		return true;
+	}
+
+	const bool bHasParent = OverrideColumn->OwnerEntry->HasParent();
+	FHierarchyTableEntryData* EntryData = OverrideColumn->OwnerEntry; 
+
+	TSharedRef<SWidget> NewWidget = SNew(SHorizontalBox)
+		+ SHorizontalBox::Slot()
+		.AutoWidth()
+		[
+			SNew(SButton)
+			.ButtonStyle(FAppStyle::Get(), "SimpleButton")
+			.IsEnabled(bHasParent)
+			.OnClicked_Lambda([EntryData]()
+				{
+					EntryData->ToggleOverridden();
+					return FReply::Handled();
+				})
+			.ContentPadding(0.0f)
+				[
+					SNew(SImage)
+					.Image_Lambda([EntryData]()
+						{
+							const bool bHasOverriddenChildren = EntryData->HasOverriddenChildren();
+
+							if (EntryData->IsOverridden())
+							{
+								if (bHasOverriddenChildren)
+								{
+									return FAppStyle::GetBrush(TEXT("DetailsView.OverrideHereInside"));
+								}
+								else
+								{
+									return FAppStyle::GetBrush(TEXT("DetailsView.OverrideHere"));
+								}
+							}
+							else
+							{
+								if (bHasOverriddenChildren)
+								{
+									return FAppStyle::GetBrush(TEXT("DetailsView.OverrideInside"));
+								}
+								else
+								{
+									return FAppStyle::GetBrush(TEXT("DetailsView.OverrideNone"));
+								}
+							}
+						})
+					.ColorAndOpacity(FSlateColor::UseForeground())
+				]
+		];
+
+	WidgetInstance->SetContent(NewWidget);
+
+	return true;
+}
\ No newline at end of file
diff --git a/Engine/Plugins/Experimental/Animation/HierarchyTable/Source/Editor/Private/HierarchyTable/Widgets/OverrideWidgetConstructor.h b/Engine/Plugins/Experimental/Animation/HierarchyTable/Source/Editor/Private/HierarchyTable/Widgets/OverrideWidgetConstructor.h
new file mode 100644
--- /dev/null
+++ b/Engine/Plugins/Experimental/Animation/HierarchyTable/Source/Editor/Private/HierarchyTable/Widgets/OverrideWidgetConstructor.h
@@ -0,0 +1,22 @@
+// Copyright Epic Games, Inc. All Rights Reserved.
+
+#pragma once
+
+#include "Elements/Interfaces/TypedElementDataStorageUiInterface.h"
+#include "OverrideWidgetConstructor.generated.h"
+
+// TODO: Can be converted into a child of FHierarchyTableWidgetConstructor
+USTRUCT()
+struct FTypedElementWidgetConstructor_Override : public FTypedElementWidgetConstructor
+{
+	GENERATED_BODY()
+
+public:
+	FTypedElementWidgetConstructor_Override();
+	~FTypedElementWidgetConstructor_Override() override = default;
+
+protected:
+	TSharedPtr<SWidget> CreateWidget(const TypedElementDataStorage::FMetaDataView& Arguments);
+	bool FinalizeWidget(ITypedElementDataStorageInterface* DataStorage, ITypedElementDataStorageUiInterface* DataStorageUi,
+		TypedElementDataStorage::RowHandle Row, const TSharedPtr<SWidget>& Widget) override;
+};
\ No newline at end of file
diff --git a/Engine/Plugins/Experimental/Animation/HierarchyTable/Source/Editor/Private/HierarchyTableAssetDefinition.cpp b/Engine/Plugins/Experimental/Animation/HierarchyTable/Source/Editor/Private/HierarchyTableAssetDefinition.cpp
new file mode 100644
--- /dev/null
+++ b/Engine/Plugins/Experimental/Animation/HierarchyTable/Source/Editor/Private/HierarchyTableAssetDefinition.cpp
@@ -0,0 +1,39 @@
+// Copyright Epic Games, Inc. All Rights Reserved.
+
+#include "HierarchyTableAssetDefinition.h"
+#include "HierarchyTable.h"
+#include "HierarchyTableEditorToolkit.h"
+
+#define LOCTEXT_NAMESPACE "HierarchyTable"
+
+FText UAssetDefinition_HierarchyTable::GetAssetDisplayName() const
+{
+	return LOCTEXT("HierarchyTable", "Hierarchy Table");
+}
+
+FLinearColor UAssetDefinition_HierarchyTable::GetAssetColor() const
+{
+	return FLinearColor(FColor::Purple);
+}
+
+TSoftClassPtr<UObject> UAssetDefinition_HierarchyTable::GetAssetClass() const
+{
+	return UHierarchyTable::StaticClass();
+}
+
+EAssetCommandResult UAssetDefinition_HierarchyTable::OpenAssets(const FAssetOpenArgs& OpenArgs) const
+{
+	if (OpenArgs.OpenMethod == EAssetOpenMethod::Edit)
+	{
+		TArray<UObject*> Assets = OpenArgs.LoadObjects<UObject>();
+		MakeShared<FHierarchyTableEditorToolkit>()->InitEditor(Assets);
+	}
+
+	return EAssetCommandResult::Handled;
+}
+
+TConstArrayView<FAssetCategoryPath> UAssetDefinition_HierarchyTable::GetAssetCategories() const
+{
+	static const auto Categories = { FAssetCategoryPath(EAssetCategoryPaths::Animation) };
+	return Categories;
+}
\ No newline at end of file
diff --git a/Engine/Plugins/Experimental/Animation/HierarchyTable/Source/Editor/Private/HierarchyTableEditorToolkit.cpp b/Engine/Plugins/Experimental/Animation/HierarchyTable/Source/Editor/Private/HierarchyTableEditorToolkit.cpp
new file mode 100644
--- /dev/null
+++ b/Engine/Plugins/Experimental/Animation/HierarchyTable/Source/Editor/Private/HierarchyTableEditorToolkit.cpp
@@ -0,0 +1,263 @@
+// Copyright Epic Games, Inc. All Rights Reserved.
+
+#include "HierarchyTableEditorToolkit.h"
+#include "Widgets/Docking/SDockTab.h"
+#include "Modules/ModuleManager.h"
+#include "PropertyEditorModule.h"
+#include "TedsOutlinerModule.h"
+#include "Elements/Framework/TypedElementRegistry.h"
+#include "Elements/Columns/TypedElementLabelColumns.h"
+#include "Elements/Columns/TypedElementTypeInfoColumns.h"
+#include "Elements/Columns/TypedElementSelectionColumns.h"
+#include "Elements/Columns/TypedElementMiscColumns.h"
+#include "SceneOutlinerPublicTypes.h"
+#include "TypedElementOutlinerMode.h"
+#include "ReferenceSkeleton.h"
+#include "Animation/Skeleton.h"
+#include "HierarchyTable/Columns/OverrideColumn.h"
+#include "HierarchyTableTypeRegistry.h"
+#include "HierarchyTableType.h"
+#include "ToolMenus.h"
+#include "Framework/Application/SlateApplication.h"
+
+#define LOCTEXT_NAMESPACE "HierarchyTableEditorToolkit"
+
+void FHierarchyTableEditorToolkit::InitEditor(const TArray<UObject*>& InObjects)
+{
+	HierarchyTable = Cast<UHierarchyTable>(InObjects[0]);
+
+	const TSharedRef<FTabManager::FLayout> Layout = FTabManager::NewLayout("HierarchyTableEditorLayout")
+		->AddArea
+		(
+			FTabManager::NewPrimaryArea()->SetOrientation(Orient_Horizontal)
+			->Split
+			(
+				FTabManager::NewStack()
+				->SetSizeCoefficient(0.7f)
+				->AddTab("HierarchyTableEditorTableTab", ETabState::OpenedTab)
+			)
+			->Split
+			(
+				FTabManager::NewStack()
+				->SetSizeCoefficient(0.3f)
+				->AddTab("HierarchyTableEditorDetailsTab", ETabState::OpenedTab)
+			)
+		);
+
+	FAssetEditorToolkit::InitAssetEditor(EToolkitMode::Standalone, {}, "HierarchyTableEditor", Layout, true, true, InObjects);
+
+	ExtendToolbar();
+}
+
+void FHierarchyTableEditorToolkit::RegisterTabSpawners(const TSharedRef<class FTabManager>& InTabManager)
+{
+	FAssetEditorToolkit::RegisterTabSpawners(InTabManager);
+
+	WorkspaceMenuCategory = InTabManager->AddLocalWorkspaceMenuCategory(LOCTEXT("HierarchyTableEditor", "Hierarchy Table Editor"));
+
+	InTabManager->RegisterTabSpawner("HierarchyTableEditorTableTab", FOnSpawnTab::CreateLambda([this](const FSpawnTabArgs&)
+		{
+			return SNew(SDockTab)
+				[
+					CreateTedsOutliner()
+				];
+		}))
+		.SetDisplayName(LOCTEXT("HierarchyTable", "Hierarchy Table"))
+		.SetGroup(WorkspaceMenuCategory.ToSharedRef());
+
+
+	FPropertyEditorModule& PropertyEditorModule = FModuleManager::GetModuleChecked<FPropertyEditorModule>("PropertyEditor");
+	FDetailsViewArgs DetailsViewArgs;
+	DetailsViewArgs.NameAreaSettings = FDetailsViewArgs::HideNameArea;
+	TSharedRef<IDetailsView> DetailsView = PropertyEditorModule.CreateDetailView(DetailsViewArgs);
+	DetailsView->SetObjects(TArray<UObject*>{ HierarchyTable });
+	InTabManager->RegisterTabSpawner("HierarchyTableEditorDetailsTab", FOnSpawnTab::CreateLambda([=](const FSpawnTabArgs&)
+		{
+			return SNew(SDockTab)
+				[
+					DetailsView
+				];
+		}))
+		.SetDisplayName(INVTEXT("Details"))
+		.SetGroup(WorkspaceMenuCategory.ToSharedRef());
+}
+
+void FHierarchyTableEditorToolkit::OnClose()
+{
+	UTypedElementRegistry* Registry = UTypedElementRegistry::GetInstance();
+	ITypedElementDataStorageInterface* DSI = Registry->GetMutableDataStorage();
+
+	for (const TTuple<int32, TypedElementDataStorage::RowHandle>& Row : EntryIndexToHandleMap)
+	{
+		DSI->RemoveRow(Row.Value);
+	}
+}
+
+void FHierarchyTableEditorToolkit::UnregisterTabSpawners(const TSharedRef<class FTabManager>& InTabManager)
+{
+	FAssetEditorToolkit::UnregisterTabSpawners(InTabManager);
+	InTabManager->UnregisterTabSpawner("HierarchyTableEditorTableTab");
+	InTabManager->UnregisterTabSpawner("HierarchyTableEditorDetailsTab");
+}
+
+TSharedRef<SWidget> FHierarchyTableEditorToolkit::CreateTedsOutliner()
+{
+	UTypedElementRegistry* Registry = UTypedElementRegistry::GetInstance();
+	checkf(Registry, TEXT("Unable to initialize the table viewer before TEDS is initialized."));
+
+	if (!Registry->AreDataStorageInterfacesSet())
+	{
+		return SNew(STextBlock)
+			.Text(INVTEXT("You need to enable the Typed Element Data Storage plugin to see the table viewer!"));
+	}
+
+	using namespace TypedElementQueryBuilder;
+
+	if (!ensure(HierarchyTable->TableType))
+	{
+		return SNullWidget::NullWidget;
+	}
+
+	const UHierarchyTableTypeHandler_Base* Handler = GetDefault<UHierarchyTableTypeRegistry>()->FindHandler(HierarchyTable->TableType);
+	if (!ensureMsgf(Handler, TEXT("Could not find handler for %s, have you forgotten to register it?"), *HierarchyTable->TableType->GetName()))
+	{
+		return SNullWidget::NullWidget;
+	}
+
+	TArray<UScriptStruct*> HierarchyTableTypeColumns = Handler->GetColumns();
+	HierarchyTableTypeColumns.Add(FTypedElementOverrideColumn::StaticStruct());
+
+	TypedElementDataStorage::FQueryDescription ColumnQueryDescription =
+		Select()
+		.ReadOnly(HierarchyTableTypeColumns)
+		.Compile();
+
+	InitialColumnQuery = Registry->GetMutableDataStorage()->RegisterQuery(MoveTemp(ColumnQueryDescription));
+
+	FSceneOutlinerInitializationOptions InitOptions;
+	InitOptions.bShowHeaderRow = true;
+	InitOptions.FilterBarOptions.bHasFilterBar = true;
+	InitOptions.OutlinerIdentifier = "HierarchyTableTedsOutliner";
+
+	FTypedElementOutlinerModeParams Params(nullptr);
+	{
+		TypedElementDataStorage::FQueryDescription RowQueryDescription =
+			Select()
+			.Where()
+			.All<FTypedElementOverrideColumn>()
+			.Compile();
+
+		Params.QueryDescription = RowQueryDescription;
+		Params.CellWidgetPurposes = TArray<FName>{ TEXT("General.Cell") };
+		Params.HierarchyData = FTypedElementOutlinerHierarchyData::GetDefaultHierarchyData();
+	}
+
+	FTedsOutlinerModule& TedsOutlinerModule = FModuleManager::GetModuleChecked<FTedsOutlinerModule>("TedsOutliner");
+
+	ITypedElementDataStorageInterface* DSI = Registry->GetMutableDataStorage();
+	static TypedElementDataStorage::TableHandle Table = DSI->FindTable(FName("Editor_HierarchyTableTable"));
+	
+	TArray<UScriptStruct*> BaseHierarchyTableTypeColumns = Handler->GetColumns();
+
+	for (int32 EntryIndex = 0; EntryIndex < HierarchyTable->TableData.Num(); ++EntryIndex)
+	{
+		FHierarchyTableEntryData* Entry = &HierarchyTable->TableData[EntryIndex];
+
+		TypedElementRowHandle Row = DSI->AddRow(Table);
+
+		FTypedElementOverrideColumn OverrideEntry;
+		OverrideEntry.OwnerEntry = Entry;
+		OverrideEntry.OwnerTable = HierarchyTable;
+		DSI->AddColumn(Row, MoveTemp(OverrideEntry));
+
+		DSI->AddColumn<FTypedElementLabelColumn>(Row, { .Label = Entry->Identifier.ToString() });
+
+		TypedElementRowHandle* ParentRow = EntryIndexToHandleMap.Find(Entry->Parent);
+		if (ParentRow)
+		{
+			DSI->AddColumn<FTypedElementParentColumn>(Row, { .Parent = *ParentRow });
+		}
+
+		for (const UScriptStruct* Column : BaseHierarchyTableTypeColumns)
+		{
+			DSI->AddColumn(Row, Column);
+		}
+
+		EntryIndexToHandleMap.Add(EntryIndex, Row);
+	}
+
+	TSharedRef<ISceneOutliner> TedsOutliner = TedsOutlinerModule.CreateTedsOutliner(InitOptions, Params, InitialColumnQuery);
+	return TedsOutliner;
+}
+
+void FHierarchyTableEditorToolkit::ExtendToolbar()
+{
+	FToolMenuOwnerScoped OwnerScoped(this);
+
+	FName ParentName;
+	static const FName MenuName = GetToolMenuToolbarName(ParentName);
+
+	UToolMenu* ToolMenu = UToolMenus::Get()->ExtendMenu(MenuName);
+	const FToolMenuInsert SectionInsertLocation("Asset", EToolMenuInsertType::After);
+	{
+		FToolMenuSection& HierarchyTableSection = ToolMenu->AddSection("HierarchyTable", LOCTEXT("HierarchyTable_ToolbarLabel", "HierarchyTable"), SectionInsertLocation);
+
+		/*
+		HierarchyTableSection.AddEntry(FToolMenuEntry::InitComboButton(
+			"AddCurve",
+			FUIAction(),
+			FNewToolMenuWidget::CreateLambda([this](const FToolMenuContext& InContext) -> TSharedRef<SWidget>
+				{
+					return SNew(SAnimCurvePicker, HierarchyTable->Skeleton)
+						.OnCurvePicked_Lambda([this](const FName SelectedCurve)
+							{
+								AddCurveEntry(SelectedCurve);
+								FSlateApplication::Get().DismissAllMenus();
+							});
+				}),
+			LOCTEXT("AddCurve_Label", "Add Curve"),
+			LOCTEXT("AddCurve_ToolTip", "Add a new curve to the hierarchy"),
+			FSlateIcon(FAppStyle::Get().GetStyleSetName(), "Icons.Plus")
+		));
+		*/
+	}
+}
+
+void FHierarchyTableEditorToolkit::AddCurveEntry(const FName CurveName)
+{
+	// TODO: Reimplement
+
+	/*
+	UTypedElementRegistry* Registry = UTypedElementRegistry::GetInstance();
+	ITypedElementDataStorageInterface* DSI = Registry->GetMutableDataStorage();
+	static TypedElementDataStorage::TableHandle Table = DSI->FindTable(FName("Editor_HierarchyTableTable"));
+
+	TypedElementDataStorage::RowHandle* RootBoneHandle = EntryIndexToHandleMap.Find(0);
+	if (!ensure(RootBoneHandle))
+	{
+		return;
+	}
+
+	TypedElementRowHandle Row = DSI->AddRow(Table);
+	DSI->AddColumn<FTypedElementOverrideColumn>(Row, { .bIsOverridden = false });
+	DSI->AddColumn<FTypedElementLabelColumn>(Row, { .Label = CurveName.ToString() });
+	DSI->AddColumn<FTypedElementMetadataColumn>(Row,
+		{
+			.OwnerTable = HierarchyTable,
+			.Children = TArray<TypedElementDataStorage::RowHandle>(),
+			.Type = EHierarchyTableMetadataType::Curve,
+			.BoneIndex = INDEX_NONE,
+			.CurveName = CurveName
+		});
+	DSI->AddColumn<FTypedElementParentColumn>(Row, { .Parent = *RootBoneHandle });
+	
+	FTypedElementMetadataColumn* MetadataColumn = DSI->GetColumn<FTypedElementMetadataColumn>(*RootBoneHandle);
+	ensure(MetadataColumn);
+	MetadataColumn->Children.Add(Row);
+
+	const UHierarchyTableTypeHandler_Base* Handler = GetDefault<UHierarchyTableTypeRegistry>()->FindHandler(HierarchyTable->TableType);
+	check(Handler);
+
+	Handler->AddData(DSI, Row, INDEX_NONE, CurveName, HierarchyTable);
+	*/
+}
\ No newline at end of file
diff --git a/Engine/Plugins/Experimental/Animation/HierarchyTable/Source/Editor/Private/HierarchyTableFactory.cpp b/Engine/Plugins/Experimental/Animation/HierarchyTable/Source/Editor/Private/HierarchyTableFactory.cpp
new file mode 100644
--- /dev/null
+++ b/Engine/Plugins/Experimental/Animation/HierarchyTable/Source/Editor/Private/HierarchyTableFactory.cpp
@@ -0,0 +1,130 @@
+// Copyright Epic Games, Inc. All Rights Reserved.
+
+#include "HierarchyTableFactory.h"
+#include "HierarchyTable.h"
+#include "ContentBrowserModule.h"
+#include "IContentBrowserSingleton.h"
+#include "Editor.h"
+#include "Animation/Skeleton.h"
+#include "Widgets/SWindow.h"
+#include "Widgets/Layout/SBorder.h"
+#include "Kismet2/SClassPickerDialog.h"
+#include "StructViewerFilter.h"
+#include "StructViewerModule.h"
+#include "HierarchyTableTypeRegistry.h"
+
+UHierarchyTableFactory::UHierarchyTableFactory()
+{
+	SupportedClass = UHierarchyTable::StaticClass();
+	bCreateNew = true;
+	TableType = nullptr;
+}
+
+UObject* UHierarchyTableFactory::FactoryCreateNew(UClass* Class, UObject* InParent, FName Name, EObjectFlags Flags, UObject* Context, FFeedbackContext* Warn)
+{
+	TObjectPtr<UHierarchyTable> HierarchyTable = NewObject<UHierarchyTable>(InParent, Class, Name, Flags, Context);
+	HierarchyTable->Skeleton = Skeleton.Get();
+	HierarchyTable->TableType = TableType;
+
+	const UHierarchyTableTypeHandler_Base* Handler = GetDefault<UHierarchyTableTypeRegistry>()->FindHandler(TableType);
+	check(Handler);
+
+	FInstancedStruct DefaultEntry = Handler->GetDefaultEntry();
+	HierarchyTable->InitializeTable(DefaultEntry);
+
+	return HierarchyTable;
+}
+
+bool UHierarchyTableFactory::ConfigureProperties()
+{
+	Skeleton = nullptr;
+
+	FContentBrowserModule& ContentBrowserModule = FModuleManager::LoadModuleChecked<FContentBrowserModule>("ContentBrowser");
+
+	FAssetPickerConfig AssetPickerConfig;
+	AssetPickerConfig.Filter.ClassPaths.Add(USkeleton::StaticClass()->GetClassPathName());
+	AssetPickerConfig.OnAssetSelected = FOnAssetSelected::CreateUObject(this, &UHierarchyTableFactory::OnSkeletonSelected);
+	AssetPickerConfig.InitialAssetViewType = EAssetViewType::List;
+
+	PickerWindow = SNew(SWindow)
+		.Title(INVTEXT("Pick Skeleton"))
+		.ClientSize(FVector2D(500, 600))
+		.SupportsMinimize(false)
+		.SupportsMaximize(false)
+		[
+			SNew(SBorder)
+			.BorderImage(FAppStyle::GetBrush("Menu.Background"))
+			[
+				ContentBrowserModule.Get().CreateAssetPicker(AssetPickerConfig)
+			]
+		];
+
+	GEditor->EditorAddModalWindow(PickerWindow.ToSharedRef());
+	PickerWindow.Reset();
+
+	if (Skeleton == nullptr)
+	{
+		return false;
+	}
+
+	FStructViewerModule& StructViewerModule = FModuleManager::LoadModuleChecked<FStructViewerModule>("StructViewer");
+
+	class FRuleFilter : public IStructViewerFilter
+	{
+	public:
+		FRuleFilter()
+		{
+		}
+
+		virtual bool IsStructAllowed(const FStructViewerInitializationOptions& InInitOptions, const UScriptStruct* InStruct, TSharedRef<FStructViewerFilterFuncs> InFilterFuncs) override
+		{
+			static const UScriptStruct* BaseStruct = FHierarchyTableType::StaticStruct();
+			return InStruct != BaseStruct && InStruct->IsChildOf(BaseStruct);
+		}
+
+		virtual bool IsUnloadedStructAllowed(const FStructViewerInitializationOptions& InInitOptions, const FSoftObjectPath& InStructPath, TSharedRef<FStructViewerFilterFuncs> InFilterFuncs) override
+		{
+			return false;
+		}
+	};
+
+	static TSharedPtr<FRuleFilter> Filter = MakeShared<FRuleFilter>();
+	FStructViewerInitializationOptions Options;
+	{
+		Options.StructFilter = Filter;
+		Options.Mode = EStructViewerMode::StructPicker;
+		Options.DisplayMode = EStructViewerDisplayMode::ListView;
+		Options.NameTypeToDisplay = EStructViewerNameTypeToDisplay::DisplayName;
+		Options.bShowNoneOption = false;
+		Options.bShowUnloadedStructs = false;
+		Options.bAllowViewOptions = false;
+	}
+
+	PickerWindow = SNew(SWindow)
+		.Title(INVTEXT("Pick Type"))
+		.ClientSize(FVector2D(500, 600))
+		.SupportsMinimize(false)
+		.SupportsMaximize(false)
+		[
+			SNew(SBorder)
+			.BorderImage(FAppStyle::GetBrush("Menu.Background"))
+		[
+			StructViewerModule.CreateStructViewer(Options, FOnStructPicked::CreateLambda([this](const UScriptStruct* ChosenStruct)
+				{
+					TableType = ChosenStruct;
+					PickerWindow->RequestDestroyWindow();
+				}))
+		]
+	];
+
+	GEditor->EditorAddModalWindow(PickerWindow.ToSharedRef());
+	PickerWindow.Reset();
+
+	return TableType != nullptr;
+}
+
+void UHierarchyTableFactory::OnSkeletonSelected(const FAssetData& SelectedAsset)
+{
+	Skeleton = Cast<USkeleton>(SelectedAsset.GetAsset());
+	PickerWindow->RequestDestroyWindow();
+}
\ No newline at end of file
diff --git a/Engine/Plugins/Experimental/Animation/HierarchyTable/Source/Editor/Private/HierarchyTableTedsFactory.cpp b/Engine/Plugins/Experimental/Animation/HierarchyTable/Source/Editor/Private/HierarchyTableTedsFactory.cpp
new file mode 100644
--- /dev/null
+++ b/Engine/Plugins/Experimental/Animation/HierarchyTable/Source/Editor/Private/HierarchyTableTedsFactory.cpp
@@ -0,0 +1,20 @@
+// Copyright Epic Games, Inc. All Rights Reserved.
+
+#include "HierarchyTableTedsFactory.h"
+
+#include "Elements/Columns/TypedElementLabelColumns.h"
+#include "Elements/Common/TypedElementDataStorageLog.h"
+#include "Elements/Framework/TypedElementQueryBuilder.h"
+#include "Elements/Framework/TypedElementRegistry.h"
+
+void UTypedElementHierarchyTableFactory::RegisterTables(ITypedElementDataStorageInterface& DataStorage)
+{
+	DataStorage.RegisterTable(
+		TTypedElementColumnTypeList<
+		FTypedElementLabelColumn>{},
+		TEXT("Editor_HierarchyTableTable"));
+}
+
+void UTypedElementHierarchyTableFactory::RegisterQueries(ITypedElementDataStorageInterface& DataStorage)
+{
+}
\ No newline at end of file
diff --git a/Engine/Plugins/Experimental/Animation/HierarchyTable/Source/Editor/Private/HierarchyTableTypeRegistry.cpp b/Engine/Plugins/Experimental/Animation/HierarchyTable/Source/Editor/Private/HierarchyTableTypeRegistry.cpp
new file mode 100644
--- /dev/null
+++ b/Engine/Plugins/Experimental/Animation/HierarchyTable/Source/Editor/Private/HierarchyTableTypeRegistry.cpp
@@ -0,0 +1,28 @@
+// Copyright Epic Games, Inc. All Rights Reserved.
+
+#pragma once
+
+#include "HierarchyTableTypeRegistry.h"
+#include "Elements/Framework/TypedElementQueryBuilder.h"
+#include "HierarchyTable/Columns/OverrideColumn.h"
+#include "HierarchyTable.h"
+
+void UHierarchyTableTypeRegistry::Register(const UScriptStruct* HierarchyTableType, const UHierarchyTableTypeHandler_Base* Handler)
+{
+	Handlers.Add(HierarchyTableType, Handler);
+}
+
+void UHierarchyTableTypeRegistry::Unregister(const UScriptStruct* HierarchyTableType)
+{
+	Handlers.Remove(HierarchyTableType);
+}
+
+const UHierarchyTableTypeHandler_Base* UHierarchyTableTypeRegistry::FindHandler(const UScriptStruct* HierarchyTableType) const
+{
+	const UHierarchyTableTypeHandler_Base* const* Result = Handlers.Find(HierarchyTableType);
+	if (Result)
+	{
+		return *Result;
+	}
+	return nullptr;
+}
\ No newline at end of file
diff --git a/Engine/Plugins/Experimental/Animation/HierarchyTable/Source/Editor/Private/HierarchyTableWidgetConstructor.cpp b/Engine/Plugins/Experimental/Animation/HierarchyTable/Source/Editor/Private/HierarchyTableWidgetConstructor.cpp
new file mode 100644
--- /dev/null
+++ b/Engine/Plugins/Experimental/Animation/HierarchyTable/Source/Editor/Private/HierarchyTableWidgetConstructor.cpp
@@ -0,0 +1,39 @@
+// Copyright Epic Games, Inc. All Rights Reserved.
+
+#include "HierarchyTableWidgetConstructor.h"
+#include "Widgets/SNullWidget.h"
+#include "Widgets/Layout/SBox.h"
+#include "Elements/Columns/TypedElementMiscColumns.h"
+#include "HierarchyTable/Columns/OverrideColumn.h"
+
+FHierarchyTableWidgetConstructor::FHierarchyTableWidgetConstructor(const UScriptStruct* InTypeInfo)
+	: Super(InTypeInfo)
+{
+}
+
+TSharedRef<SWidget> FHierarchyTableWidgetConstructor::CreateWidget(FHierarchyTableEntryData* EntryData)
+{
+	return SNullWidget::NullWidget;
+}
+
+TSharedPtr<SWidget> FHierarchyTableWidgetConstructor::CreateWidget(const TypedElementDataStorage::FMetaDataView& Arguments)
+{
+	return SNew(SBox)
+		.HAlign(HAlign_Left)
+		.VAlign(VAlign_Center);
+}
+
+bool FHierarchyTableWidgetConstructor::FinalizeWidget(ITypedElementDataStorageInterface* DataStorage, ITypedElementDataStorageUiInterface* DataStorageUi,
+	TypedElementDataStorage::RowHandle Row, const TSharedPtr<SWidget>& Widget)
+{
+	TSharedPtr<SBox> WidgetInstance = StaticCastSharedPtr<SBox>(Widget);
+
+	// Row is not actually the row we want, its contained inside of a row reference, I don't know why things are done this way.
+	TypedElementDataStorage::RowHandle TargetRow = DataStorage->GetColumn<FTypedElementRowReferenceColumn>(Row)->Row;
+	FTypedElementOverrideColumn* OverrideColumn = DataStorage->GetColumn<FTypedElementOverrideColumn>(TargetRow);
+
+	TSharedRef<SWidget> ActualWidget = CreateWidget(OverrideColumn->OwnerEntry);
+	WidgetInstance->SetContent(ActualWidget);
+
+	return true;
+}
\ No newline at end of file
diff --git a/Engine/Plugins/Experimental/Animation/HierarchyTable/Source/Editor/Public/HierarchyTable/Columns/OverrideColumn.h b/Engine/Plugins/Experimental/Animation/HierarchyTable/Source/Editor/Public/HierarchyTable/Columns/OverrideColumn.h
new file mode 100644
--- /dev/null
+++ b/Engine/Plugins/Experimental/Animation/HierarchyTable/Source/Editor/Public/HierarchyTable/Columns/OverrideColumn.h
@@ -0,0 +1,13 @@
+// Copyright Epic Games, Inc. All Rights Reserved.
+
+#pragma once
+
+#include "HierarchyTableProxyColumn.h"
+
+#include "OverrideColumn.generated.h"
+
+USTRUCT()
+struct FTypedElementOverrideColumn final : public FHierarchyTableProxyColumn
+{
+	GENERATED_BODY()
+};
\ No newline at end of file
diff --git a/Engine/Plugins/Experimental/Animation/HierarchyTable/Source/Editor/Public/HierarchyTableAssetDefinition.h b/Engine/Plugins/Experimental/Animation/HierarchyTable/Source/Editor/Public/HierarchyTableAssetDefinition.h
new file mode 100644
--- /dev/null
+++ b/Engine/Plugins/Experimental/Animation/HierarchyTable/Source/Editor/Public/HierarchyTableAssetDefinition.h
@@ -0,0 +1,22 @@
+// Copyright Epic Games, Inc. All Rights Reserved.
+
+#pragma once
+
+#include "AssetDefinitionDefault.h"
+
+#include "HierarchyTableAssetDefinition.generated.h"
+
+UCLASS()
+class UAssetDefinition_HierarchyTable : public UAssetDefinitionDefault
+{
+	GENERATED_BODY()
+
+public:
+	// Begin UAssetDefinition
+	virtual FText GetAssetDisplayName() const override;
+	virtual FLinearColor GetAssetColor() const override;
+	virtual TSoftClassPtr<UObject> GetAssetClass() const override;
+	virtual EAssetCommandResult OpenAssets(const FAssetOpenArgs& OpenArgs) const override;
+	virtual TConstArrayView<FAssetCategoryPath> GetAssetCategories() const override;
+	// End UAssetDefinition
+};
\ No newline at end of file
diff --git a/Engine/Plugins/Experimental/Animation/HierarchyTable/Source/Editor/Public/HierarchyTableEditorModule.h b/Engine/Plugins/Experimental/Animation/HierarchyTable/Source/Editor/Public/HierarchyTableEditorModule.h
new file mode 100644
--- /dev/null
+++ b/Engine/Plugins/Experimental/Animation/HierarchyTable/Source/Editor/Public/HierarchyTableEditorModule.h
@@ -0,0 +1,17 @@
+// Copyright Epic Games, Inc. All Rights Reserved.
+
+#pragma once
+
+#include "Templates/SharedPointer.h"
+#include "Modules/ModuleInterface.h"
+
+class FHierarchyTableEditorModule : public IModuleInterface
+{
+public:
+	virtual void StartupModule() override;
+	
+	virtual void ShutdownModule() override;
+
+private:
+	TSharedPtr<class FHierarchyTableAssetTypeActions> HierarchyTableAssetTypeActions;
+};
\ No newline at end of file
diff --git a/Engine/Plugins/Experimental/Animation/HierarchyTable/Source/Editor/Public/HierarchyTableEditorToolkit.h b/Engine/Plugins/Experimental/Animation/HierarchyTable/Source/Editor/Public/HierarchyTableEditorToolkit.h
new file mode 100644
--- /dev/null
+++ b/Engine/Plugins/Experimental/Animation/HierarchyTable/Source/Editor/Public/HierarchyTableEditorToolkit.h
@@ -0,0 +1,38 @@
+// Copyright Epic Games, Inc. All Rights Reserved.
+
+#pragma once
+
+#include "CoreMinimal.h"
+#include "HierarchyTable.h"
+#include "Toolkits/AssetEditorToolkit.h"
+#include "Elements/Common/TypedElementHandles.h"
+
+class FHierarchyTableEditorToolkit : public FAssetEditorToolkit
+{
+public:
+	void InitEditor(const TArray<UObject*>& InObjects);
+	virtual void OnClose() override;
+
+	void RegisterTabSpawners(const TSharedRef<class FTabManager>& TabManager) override;
+	void UnregisterTabSpawners(const TSharedRef<class FTabManager>& TabManager) override;
+
+	FName GetToolkitFName() const override { return "HierarchyTableEditor"; }
+	FText GetBaseToolkitName() const override { return INVTEXT("Hierarchy Table Editor"); }
+	FString GetWorldCentricTabPrefix() const override { return "Hierarchy Table "; }
+	FLinearColor GetWorldCentricTabColorScale() const override { return {}; }
+
+	void ExtendToolbar();
+
+private:
+	void AddCurveEntry(const FName CurveName);
+
+	TSharedRef<SWidget> CreateTedsOutliner();
+
+	UHierarchyTable* HierarchyTable;
+
+	TMap<int32, TypedElementDataStorage::RowHandle> EntryIndexToHandleMap;
+
+	TypedElementDataStorage::QueryHandle InitialColumnQuery;
+
+	//TypedElementDataStorage::QueryHandle CollectDataQuery;
+};
diff --git a/Engine/Plugins/Experimental/Animation/HierarchyTable/Source/Editor/Public/HierarchyTableFactory.h b/Engine/Plugins/Experimental/Animation/HierarchyTable/Source/Editor/Public/HierarchyTableFactory.h
new file mode 100644
--- /dev/null
+++ b/Engine/Plugins/Experimental/Animation/HierarchyTable/Source/Editor/Public/HierarchyTableFactory.h
@@ -0,0 +1,34 @@
+// Copyright Epic Games, Inc. All Rights Reserved.
+
+#pragma once
+
+#include "CoreMinimal.h"
+#include "Factories/Factory.h"
+#include "HierarchyTableType.h"
+
+#include "HierarchyTableFactory.generated.h"
+
+class UScriptStruct;
+
+UCLASS()
+class UHierarchyTableFactory : public UFactory
+{
+	GENERATED_BODY()
+
+	UPROPERTY()
+	TWeakObjectPtr<class USkeleton>	Skeleton;
+
+	const UScriptStruct* TableType;
+
+public:
+	UHierarchyTableFactory();
+
+	// UFactory interface
+	UObject* FactoryCreateNew(UClass* Class, UObject* InParent, FName Name, EObjectFlags Flags, UObject* Context, FFeedbackContext* Warn) override;
+	virtual bool ConfigureProperties() override;
+
+private:
+	void OnSkeletonSelected(const FAssetData& SelectedAsset);
+
+	TSharedPtr<class SWindow> PickerWindow;
+};
diff --git a/Engine/Plugins/Experimental/Animation/HierarchyTable/Source/Editor/Public/HierarchyTableProxyColumn.h b/Engine/Plugins/Experimental/Animation/HierarchyTable/Source/Editor/Public/HierarchyTableProxyColumn.h
new file mode 100644
--- /dev/null
+++ b/Engine/Plugins/Experimental/Animation/HierarchyTable/Source/Editor/Public/HierarchyTableProxyColumn.h
@@ -0,0 +1,21 @@
+// Copyright Epic Games, Inc. All Rights Reserved.
+
+#pragma once
+
+#include "Elements/Common/TypedElementCommonTypes.h"
+
+#include "HierarchyTableProxyColumn.generated.h"
+
+class UHierarchyTable;
+struct FHierarchyTableEntryData;
+
+USTRUCT()
+struct FHierarchyTableProxyColumn : public FTypedElementDataStorageColumn
+{
+	GENERATED_BODY()
+
+public:
+	UHierarchyTable* OwnerTable;
+
+	FHierarchyTableEntryData* OwnerEntry;
+};
\ No newline at end of file
diff --git a/Engine/Plugins/Experimental/Animation/HierarchyTable/Source/Editor/Public/HierarchyTableTedsFactory.h b/Engine/Plugins/Experimental/Animation/HierarchyTable/Source/Editor/Public/HierarchyTableTedsFactory.h
new file mode 100644
--- /dev/null
+++ b/Engine/Plugins/Experimental/Animation/HierarchyTable/Source/Editor/Public/HierarchyTableTedsFactory.h
@@ -0,0 +1,21 @@
+// Copyright Epic Games, Inc. All Rights Reserved.
+
+#pragma once
+
+#include "Elements/Interfaces/TypedElementDataStorageFactory.h"
+#include "Elements/Interfaces/TypedElementQueryStorageInterfaces.h"
+#include "UObject/ObjectMacros.h"
+
+#include "HierarchyTableTedsFactory.generated.h"
+
+UCLASS()
+class UTypedElementHierarchyTableFactory : public UTypedElementDataStorageFactory
+{
+	GENERATED_BODY()
+
+public:
+	~UTypedElementHierarchyTableFactory() override = default;
+
+	void RegisterTables(ITypedElementDataStorageInterface& DataStorage) override;
+	void RegisterQueries(ITypedElementDataStorageInterface& DataStorage) override;
+};
diff --git a/Engine/Plugins/Experimental/Animation/HierarchyTable/Source/Editor/Public/HierarchyTableTypeRegistry.h b/Engine/Plugins/Experimental/Animation/HierarchyTable/Source/Editor/Public/HierarchyTableTypeRegistry.h
new file mode 100644
--- /dev/null
+++ b/Engine/Plugins/Experimental/Animation/HierarchyTable/Source/Editor/Public/HierarchyTableTypeRegistry.h
@@ -0,0 +1,42 @@
+// Copyright Epic Games, Inc. All Rights Reserved.
+
+#pragma once
+
+#include "UObject/Object.h"
+#include "Elements/Interfaces/TypedElementDataStorageInterface.h"
+#include "Elements/Common/TypedElementHandles.h"
+#include "HierarchyTableType.h"
+#include "StructUtils/InstancedStruct.h"
+#include "HierarchyTableTypeRegistry.generated.h"
+
+class UHierarchyTable;
+class UScriptStruct;
+struct FHierarchyTableEntryData;
+
+UCLASS(Abstract)
+class HIERARCHYTABLEEDITOR_API UHierarchyTableTypeHandler_Base : public UObject
+{
+	GENERATED_BODY()
+
+public:
+	virtual FInstancedStruct GetDefaultEntry() const { return FInstancedStruct(); }
+
+	virtual void InitializeTable(TObjectPtr<class UHierarchyTable> InHierarchyTable) const {};
+
+	virtual TArray<UScriptStruct*> GetColumns() const { return {}; };
+};
+
+UCLASS()
+class HIERARCHYTABLEEDITOR_API UHierarchyTableTypeRegistry : public UObject
+{
+	GENERATED_BODY()
+
+public:
+	void Register(const UScriptStruct* HierarchyTableType, const UHierarchyTableTypeHandler_Base* Handler);
+	void Unregister(const UScriptStruct* HierarchyTableType);
+
+	const UHierarchyTableTypeHandler_Base* FindHandler(const UScriptStruct* HierarchyTableType) const;
+
+private:
+	TMap<const UScriptStruct*, const UHierarchyTableTypeHandler_Base*> Handlers;
+};
\ No newline at end of file
diff --git a/Engine/Plugins/Experimental/Animation/HierarchyTable/Source/Editor/Public/HierarchyTableWidgetConstructor.h b/Engine/Plugins/Experimental/Animation/HierarchyTable/Source/Editor/Public/HierarchyTableWidgetConstructor.h
new file mode 100644
--- /dev/null
+++ b/Engine/Plugins/Experimental/Animation/HierarchyTable/Source/Editor/Public/HierarchyTableWidgetConstructor.h
@@ -0,0 +1,31 @@
+// Copyright Epic Games, Inc. All Rights Reserved.
+
+#pragma once
+
+#include "Elements/Interfaces/TypedElementDataStorageUiInterface.h"
+
+#include "HierarchyTableWidgetConstructor.generated.h"
+
+struct FHierarchyTableEntryData;
+
+USTRUCT()
+struct HIERARCHYTABLEEDITOR_API FHierarchyTableWidgetConstructor : public FTypedElementWidgetConstructor
+{
+	GENERATED_BODY()
+
+public:
+	explicit FHierarchyTableWidgetConstructor() : Super(StaticStruct()) {}
+	explicit FHierarchyTableWidgetConstructor(const UScriptStruct* InTypeInfo);
+
+	virtual ~FHierarchyTableWidgetConstructor() override = default;
+
+	virtual TSharedRef<SWidget> CreateWidget(FHierarchyTableEntryData* EntryData);
+
+protected:
+	// Begin FHierarchyTableWidgetConstructor
+	TSharedPtr<SWidget> CreateWidget(const TypedElementDataStorage::FMetaDataView& Arguments);
+
+	bool FinalizeWidget(ITypedElementDataStorageInterface* DataStorage, ITypedElementDataStorageUiInterface* DataStorageUi,
+		TypedElementDataStorage::RowHandle Row, const TSharedPtr<SWidget>& Widget) override;
+	// End FHierarchyTableWidgetConstructor
+};
\ No newline at end of file
diff --git a/Engine/Plugins/Experimental/Animation/HierarchyTable/Source/Runtime/HierarchyTableRuntime.Build.cs b/Engine/Plugins/Experimental/Animation/HierarchyTable/Source/Runtime/HierarchyTableRuntime.Build.cs
new file mode 100644
--- /dev/null
+++ b/Engine/Plugins/Experimental/Animation/HierarchyTable/Source/Runtime/HierarchyTableRuntime.Build.cs
@@ -0,0 +1,25 @@
+// Copyright Epic Games, Inc. All Rights Reserved.
+
+using UnrealBuildTool;
+
+public class HierarchyTableRuntime : ModuleRules
+{
+	public HierarchyTableRuntime(ReadOnlyTargetRules Target) : base(Target)
+	{
+		PublicDependencyModuleNames.AddRange(
+			new string[]
+			{
+				"Core",
+				"CoreUObject",
+				"Engine",
+				"AnimationCore"
+			}
+		);
+
+		PrivateDependencyModuleNames.AddRange(
+			new string[]
+            {
+			}
+		);
+	}
+}
diff --git a/Engine/Plugins/Experimental/Animation/HierarchyTable/Source/Runtime/Private/HierachyTableRuntimeModule.cpp b/Engine/Plugins/Experimental/Animation/HierarchyTable/Source/Runtime/Private/HierachyTableRuntimeModule.cpp
new file mode 100644
--- /dev/null
+++ b/Engine/Plugins/Experimental/Animation/HierarchyTable/Source/Runtime/Private/HierachyTableRuntimeModule.cpp
@@ -0,0 +1,14 @@
+// Copyright Epic Games, Inc. All Rights Reserved.
+
+#include "HierarchyTableRuntimeModule.h"
+#include "Modules/ModuleManager.h"
+
+void FHierarchyTableRuntimeModule::StartupModule()
+{
+}
+
+void FHierarchyTableRuntimeModule::ShutdownModule()
+{
+}
+
+IMPLEMENT_MODULE(FHierarchyTableRuntimeModule, HierarchyTableRuntime)
diff --git a/Engine/Plugins/Experimental/Animation/HierarchyTable/Source/Runtime/Private/HierarchyTable.cpp b/Engine/Plugins/Experimental/Animation/HierarchyTable/Source/Runtime/Private/HierarchyTable.cpp
new file mode 100644
--- /dev/null
+++ b/Engine/Plugins/Experimental/Animation/HierarchyTable/Source/Runtime/Private/HierarchyTable.cpp
@@ -0,0 +1,107 @@
+// Copyright Epic Games, Inc. All Rights Reserved.
+
+#include "HierarchyTable.h"
+#include "Animation/Skeleton.h"
+#include "ReferenceSkeleton.h"
+
+UHierarchyTable::UHierarchyTable()
+{}
+
+void UHierarchyTable::InitializeTable(FInstancedStruct DefaultEntry)
+{
+	check(Skeleton);
+	check(TableData.Num() == 0);
+
+	FReferenceSkeleton RefSkeleton = Skeleton->GetReferenceSkeleton();
+
+	for (int32 BoneIndex = 0; BoneIndex < RefSkeleton.GetNum(); ++BoneIndex)
+	{
+		FHierarchyTableEntryData EntryData;
+		EntryData.Parent = RefSkeleton.GetParentIndex(BoneIndex);
+		EntryData.Identifier = RefSkeleton.GetBoneName(BoneIndex);
+		EntryData.EntryType = EHierarchyTableEntryType::Bone;
+		EntryData.Payload = BoneIndex == 0 ? DefaultEntry : TOptional<FInstancedStruct>();
+		EntryData.OwnerTable = this;
+
+		TableData.Add(EntryData);
+	}
+}
+
+bool FHierarchyTableEntryData::HasOverriddenChildren() const
+{
+	return IsOverriddenOrHasOverriddenChildren(false);
+}
+
+bool FHierarchyTableEntryData::IsOverriddenOrHasOverriddenChildren(const bool bIncludeSelf) const
+{
+	if (bIncludeSelf && IsOverridden())
+	{
+		return true;
+	}
+
+	TArray<const FHierarchyTableEntryData*> Children = OwnerTable->GetChildren(*this);
+	for (const FHierarchyTableEntryData* Child : Children)
+	{
+		if (Child->IsOverriddenOrHasOverriddenChildren(true))
+		{
+			return true;
+		}
+	}
+
+	return false;
+}
+
+void FHierarchyTableEntryData::ToggleOverridden()
+{
+	if (Payload.IsSet())
+	{
+		Payload.Reset();
+	}
+	else
+	{
+		Payload = *GetFromClosestAncestor();
+	}
+}
+
+const FInstancedStruct* FHierarchyTableEntryData::GetActualValue() const
+{
+	return IsOverridden() ? Payload.GetPtrOrNull() : GetFromClosestAncestor();
+}
+
+const FHierarchyTableEntryData* FHierarchyTableEntryData::GetClosestAncestor()
+{
+	return IsOverridden() ? this : OwnerTable->TableData[Parent].GetClosestAncestor();
+}
+
+const FInstancedStruct* FHierarchyTableEntryData::GetFromClosestAncestor() const
+{
+	check(Parent != INDEX_NONE);
+
+	const FHierarchyTableEntryData& ParentEntry = OwnerTable->TableData[Parent];
+	return ParentEntry.GetActualValue();
+}
+
+TArray<const FHierarchyTableEntryData*> UHierarchyTable::GetChildren(const FHierarchyTableEntryData& Parent) const
+{
+	int32 ParentIndex = TableData.IndexOfByPredicate([Parent](const FHierarchyTableEntryData& Candidate)
+		{
+			return Candidate.Identifier == Parent.Identifier;
+		});
+
+	if (ParentIndex == INDEX_NONE)
+	{
+		return TArray<const FHierarchyTableEntryData*>();
+	}
+
+	TArray<const FHierarchyTableEntryData*> Children;
+
+	for (const FHierarchyTableEntryData& Entry : TableData)
+	{
+		if (Entry.Parent == ParentIndex)
+		{
+			Children.Add(&Entry);
+		}
+	}
+
+	return Children;
+}
\ No newline at end of file
diff --git a/Engine/Plugins/Experimental/Animation/HierarchyTable/Source/Runtime/Public/HierarchyTable.h b/Engine/Plugins/Experimental/Animation/HierarchyTable/Source/Runtime/Public/HierarchyTable.h
new file mode 100644
--- /dev/null
+++ b/Engine/Plugins/Experimental/Animation/HierarchyTable/Source/Runtime/Public/HierarchyTable.h
@@ -0,0 +1,103 @@
+// Copyright Epic Games, Inc. All Rights Reserved.
+
+#pragma once
+
+#include "CoreMinimal.h"
+#include "UObject/NoExportTypes.h"
+#include "HierarchyTableType.h"
+#include "Templates/SubclassOf.h"
+#include "Templates/SharedPointer.h"
+#include "StructUtils/InstancedStruct.h"
+
+#include "HierarchyTable.generated.h"
+
+class USkeleton;
+
+UENUM()
+enum class EHierarchyTableEntryType
+{
+	Bone,
+	Curve,
+};
+
+class UHierarchyTable;
+
+USTRUCT()
+struct HIERARCHYTABLERUNTIME_API FHierarchyTableEntryData
+{
+	GENERATED_BODY()
+
+public:
+	friend UHierarchyTable;
+
+	UPROPERTY()
+	TObjectPtr<UHierarchyTable> OwnerTable;
+
+	UPROPERTY()
+	EHierarchyTableEntryType EntryType;
+
+	UPROPERTY()
+	FName Identifier;
+
+	UPROPERTY()
+	int32 Parent;
+
+	UPROPERTY()
+	TOptional<FInstancedStruct> Payload;
+
+public:
+	bool HasParent() const { return Parent != INDEX_NONE; }
+	bool IsOverridden() const { return Payload.IsSet(); }
+	bool HasOverriddenChildren() const;
+
+	void ToggleOverridden();
+
+	template <typename T>
+	const T* GetValue() const
+	{
+		const FInstancedStruct* StructPtr = GetActualValue();
+		check(StructPtr->IsValid());
+		return StructPtr->GetPtr<T>();
+	}
+
+	template <typename T>
+	T* GetMutableValue()
+	{
+		check(IsOverridden());
+		return Payload.GetValue().GetMutablePtr<T>();
+	}
+
+	const FHierarchyTableEntryData* GetClosestAncestor();
+
+private:
+	// TODO: This needs to be cached in some way as it potentially goes up the entire hierarchy until it finds an ancestor with an overridden value
+	// This is called each time a widget is ticked so grows exponentially with the height of the hierarchy. This needs addressing at some point
+	// but cached value must be updated when any of its direct ancestors' value is updated.
+	const FInstancedStruct* GetActualValue() const;
+
+	const FInstancedStruct* GetFromClosestAncestor() const;
+
+	bool IsOverriddenOrHasOverriddenChildren(const bool bIncludeSelf) const;
+};
+
+UCLASS()
+class HIERARCHYTABLERUNTIME_API UHierarchyTable : public UObject
+{
+	GENERATED_BODY()
+
+public:
+	UHierarchyTable();
+
+	UPROPERTY(VisibleAnywhere, Category=Default)
+	TObjectPtr<USkeleton> Skeleton;
+
+	UPROPERTY(VisibleAnywhere, Category = Default)
+	TObjectPtr<const UScriptStruct> TableType;
+
+	UPROPERTY()
+	TArray<FHierarchyTableEntryData> TableData;
+
+	void InitializeTable(FInstancedStruct DefaultEntry);
+
+	TArray<const FHierarchyTableEntryData*> GetChildren(const FHierarchyTableEntryData& Parent) const;
+};
diff --git a/Engine/Plugins/Experimental/Animation/HierarchyTable/Source/Runtime/Public/HierarchyTableRuntimeModule.h b/Engine/Plugins/Experimental/Animation/HierarchyTable/Source/Runtime/Public/HierarchyTableRuntimeModule.h
new file mode 100644
--- /dev/null
+++ b/Engine/Plugins/Experimental/Animation/HierarchyTable/Source/Runtime/Public/HierarchyTableRuntimeModule.h
@@ -0,0 +1,13 @@
+// Copyright Epic Games, Inc. All Rights Reserved.
+
+#pragma once
+
+#include "Modules/ModuleInterface.h"
+
+class FHierarchyTableRuntimeModule : public IModuleInterface
+{
+public:
+
+	virtual void StartupModule() override;
+	virtual void ShutdownModule() override;
+};
\ No newline at end of file
diff --git a/Engine/Plugins/Experimental/Animation/HierarchyTable/Source/Runtime/Public/HierarchyTableType.h b/Engine/Plugins/Experimental/Animation/HierarchyTable/Source/Runtime/Public/HierarchyTableType.h
new file mode 100644
--- /dev/null
+++ b/Engine/Plugins/Experimental/Animation/HierarchyTable/Source/Runtime/Public/HierarchyTableType.h
@@ -0,0 +1,11 @@
+// Copyright Epic Games, Inc. All Rights Reserved.
+
+#pragma once
+
+#include "HierarchyTableType.generated.h"
+
+USTRUCT()
+struct HIERARCHYTABLERUNTIME_API FHierarchyTableType
+{
+	GENERATED_BODY()
+};
\ No newline at end of file
