Fix missing #undef LOCTEXT_NAMESPACE

diff --git a/Engine/Plugins/Experimental/Animation/HierarchyTable/Source/Editor/Private/HierarchyTableAssetDefinition.cpp b/Engine/Plugins/Experimental/Animation/HierarchyTable/Source/Editor/Private/HierarchyTableAssetDefinition.cpp
--- a/Engine/Plugins/Experimental/Animation/HierarchyTable/Source/Editor/Private/HierarchyTableAssetDefinition.cpp
+++ b/Engine/Plugins/Experimental/Animation/HierarchyTable/Source/Editor/Private/HierarchyTableAssetDefinition.cpp
@@ -36,4 +36,6 @@ TConstArrayView<FAssetCategoryPath> UAssetDefinition_HierarchyTable::GetAssetCat
 {
 	static const auto Categories = { FAssetCategoryPath(EAssetCategoryPaths::Animation) };
 	return Categories;
-}
\ No newline at end of file
+}
+
+#undef LOCTEXT_NAMESPACE
\ No newline at end of file
diff --git a/Engine/Plugins/Experimental/Animation/HierarchyTable/Source/Editor/Private/HierarchyTableEditorToolkit.cpp b/Engine/Plugins/Experimental/Animation/HierarchyTable/Source/Editor/Private/HierarchyTableEditorToolkit.cpp
--- a/Engine/Plugins/Experimental/Animation/HierarchyTable/Source/Editor/Private/HierarchyTableEditorToolkit.cpp
+++ b/Engine/Plugins/Experimental/Animation/HierarchyTable/Source/Editor/Private/HierarchyTableEditorToolkit.cpp
@@ -260,4 +260,6 @@ void FHierarchyTableEditorToolkit::AddCurveEntry(const FName CurveName)
 
 	Handler->AddData(DSI, Row, INDEX_NONE, CurveName, HierarchyTable);
 	*/
-}
\ No newline at end of file
+}
+
+#undef LOCTEXT_NAMESPACE
\ No newline at end of file
