Add icons to hierarchy table

diff --git a/Engine/Plugins/Experimental/Animation/HierarchyTable/Source/Editor/Private/HierarchyTableEditorToolkit.cpp b/Engine/Plugins/Experimental/Animation/HierarchyTable/Source/Editor/Private/HierarchyTableEditorToolkit.cpp
--- a/Engine/Plugins/Experimental/Animation/HierarchyTable/Source/Editor/Private/HierarchyTableEditorToolkit.cpp
+++ b/Engine/Plugins/Experimental/Animation/HierarchyTable/Source/Editor/Private/HierarchyTableEditorToolkit.cpp
@@ -176,7 +176,7 @@ TSharedRef<SWidget> FHierarchyTableEditorToolkit::CreateTedsOutliner()
 		OverrideEntry.OwnerTable = HierarchyTable;
 		DSI->AddColumn(Row, MoveTemp(OverrideEntry));
 
-		DSI->AddColumn<FTypedElementLabelColumn>(Row, { .Label = CreateRowLabel(Entry->Identifier, Entry->EntryType) });
+		DSI->AddColumn<FTypedElementLabelColumn>(Row, { .Label = Entry->Identifier.ToString() });
 
 		RowHandle* ParentRow = EntryIndexToHandleMap.Find(Entry->Parent);
 		if (ParentRow)
@@ -346,7 +346,7 @@ void FHierarchyTableEditorToolkit::AddEntry(const FName Identifier, const EHiera
 
 	// Ideally would read the label directly from the HT entry struct instead of storing it itself
 	// but this is a built-in TEDS column type that is used the the tabel viewer widget.
-	DSI->AddColumn<FTypedElementLabelColumn>(Row, { .Label = CreateRowLabel(EntryData.Identifier, EntryData.EntryType) });
+	DSI->AddColumn<FTypedElementLabelColumn>(Row, { .Label = EntryData.Identifier.ToString() });
 
 	const RowHandle* ParentRow = EntryIndexToHandleMap.Find(EntryData.Parent);
 	if (ParentRow)
@@ -365,17 +365,4 @@ void FHierarchyTableEditorToolkit::AddEntry(const FName Identifier, const EHiera
 	EntryIndexToHandleMap.Add(EntryIndex, Row);
 }
 
-FString FHierarchyTableEditorToolkit::CreateRowLabel(const FName EntryIdentifier, const EHierarchyTableEntryType EntryType)
-{
-	switch (EntryType)
-	{
-		case EHierarchyTableEntryType::Curve:
-			return FString::Format(TEXT("{0} (Curve)"), { EntryIdentifier.ToString() });
-		case EHierarchyTableEntryType::Attribute:
-			return FString::Format(TEXT("{0} (Attribute)"), { EntryIdentifier.ToString() });
-		default:
-			return EntryIdentifier.ToString();
-	}
-}
-
 #undef LOCTEXT_NAMESPACE
\ No newline at end of file
diff --git a/Engine/Plugins/Experimental/Animation/HierarchyTable/Source/Editor/Private/HierarchyTableLabelWidget.cpp b/Engine/Plugins/Experimental/Animation/HierarchyTable/Source/Editor/Private/HierarchyTableLabelWidget.cpp
new file mode 100644
--- /dev/null
+++ b/Engine/Plugins/Experimental/Animation/HierarchyTable/Source/Editor/Private/HierarchyTableLabelWidget.cpp
@@ -0,0 +1,100 @@
+// Copyright Epic Games, Inc. All Rights Reserved.
+
+#include "HierarchyTableLabelWidget.h"
+
+#include "HierarchyTable/Columns/OverrideColumn.h"
+#include "Elements/Columns/TypedElementCompatibilityColumns.h"
+#include "Elements/Columns/TypedElementLabelColumns.h"
+#include "Elements/Framework/TypedElementAttributeBinding.h"
+#include "Widgets/Images/SImage.h"
+#include "Widgets/Layout/SSpacer.h"
+#include "Widgets/Text/STextBlock.h"
+#include "Widgets/SBoxPanel.h"
+#include "HierarchyTable.h"
+
+#define LOCTEXT_NAMESPACE "HierarchyTableLabelWidget"
+
+void UHierarchyTableLabelWidgetFactory::RegisterWidgetConstructors(IEditorDataStorageProvider& DataStorage,
+	IEditorDataStorageUiProvider& DataStorageUi) const
+{
+	using namespace UE::Editor::DataStorage::Queries;
+	DataStorageUi.RegisterWidgetFactory<FHierarchyTableLabelWidgetConstructor>(
+		TEXT("General.RowLabel"),
+		TColumn<FTypedElementLabelColumn>() && TColumn<FTypedElementOverrideColumn>());
+}
+
+FHierarchyTableLabelWidgetConstructor::FHierarchyTableLabelWidgetConstructor()
+	: FSimpleWidgetConstructor(StaticStruct())
+{
+}
+
+FHierarchyTableLabelWidgetConstructor::FHierarchyTableLabelWidgetConstructor(const UScriptStruct* TypeInfo)
+	: FSimpleWidgetConstructor(TypeInfo)
+{
+}
+
+TSharedPtr<SWidget> FHierarchyTableLabelWidgetConstructor::CreateWidget(
+	IEditorDataStorageProvider* DataStorage,
+	IEditorDataStorageUiProvider* DataStorageUi,
+	RowHandle TargetRow,
+	RowHandle WidgetRow,
+	const UE::Editor::DataStorage::FMetaDataView& Arguments)
+{
+	using namespace UE::Editor::DataStorage;
+	if(DataStorage->IsRowAvailable(TargetRow))
+	{
+		FAttributeBinder Binder(TargetRow, DataStorage);
+
+		FTypedElementOverrideColumn* OverrideColumn = DataStorage->GetColumn<FTypedElementOverrideColumn>(TargetRow);
+		const FHierarchyTableEntryData& EntryData = OverrideColumn->OwnerTable->TableData[OverrideColumn->OwnerEntryIndex];
+
+		const FSlateBrush* LabelIcon = nullptr;
+		FSlateColor LabelIconColor = FSlateColor::UseForeground();
+
+		switch (EntryData.EntryType)
+		{
+			case EHierarchyTableEntryType::Bone:
+				LabelIcon = FAppStyle::GetBrush("SkeletonTree.Bone");
+				break;
+			case EHierarchyTableEntryType::Curve:
+				LabelIcon = FAppStyle::GetBrush("AnimGraph.Attribute.Curves.Icon");
+				LabelIconColor = FAppStyle::GetSlateColor("AnimGraph.Attribute.Curves.Color");
+				break;
+			case EHierarchyTableEntryType::Attribute:
+				LabelIcon = FAppStyle::GetBrush("AnimGraph.Attribute.Attributes.Icon");
+				LabelIconColor = FAppStyle::GetSlateColor("AnimGraph.Attribute.Attributes.Color");
+				break;
+		}
+
+		return SNew(SHorizontalBox)
+			+SHorizontalBox::Slot()
+				.AutoWidth()
+			[
+				SNew(SImage)
+					.Image(LabelIcon)
+					.ColorAndOpacity(LabelIconColor)
+			]
+			+SHorizontalBox::Slot()
+				.AutoWidth()
+			[
+				SNew(SSpacer)
+					.Size(FVector2D(5.0f, 0.0f))
+			]
+			+SHorizontalBox::Slot()
+				.FillWidth(1.0f)
+			[
+				SNew(STextBlock)
+					.Text(Binder.BindText(&FTypedElementLabelColumn::Label))
+					.ToolTipText(Binder.BindText(&FTypedElementLabelColumn::Label))
+			];
+	}
+	else
+	{
+		return SNew(STextBlock)
+			.Text(LOCTEXT("MissingRowReferenceColumn", "Unable to retrieve row reference."));
+	}
+	
+	
+}
+
+#undef LOCTEXT_NAMESPACE
diff --git a/Engine/Plugins/Experimental/Animation/HierarchyTable/Source/Editor/Public/HierarchyTableEditorToolkit.h b/Engine/Plugins/Experimental/Animation/HierarchyTable/Source/Editor/Public/HierarchyTableEditorToolkit.h
--- a/Engine/Plugins/Experimental/Animation/HierarchyTable/Source/Editor/Public/HierarchyTableEditorToolkit.h
+++ b/Engine/Plugins/Experimental/Animation/HierarchyTable/Source/Editor/Public/HierarchyTableEditorToolkit.h
@@ -24,8 +24,6 @@ public:
 	void ExtendToolbar();
 
 private:
-	FString CreateRowLabel(const FName EntryIdentifier, const EHierarchyTableEntryType EntryType);
-
 	void AddEntry(const FName Identifier, const EHierarchyTableEntryType EntryType);
 
 	TSharedRef<SWidget> CreateTedsOutliner();
@@ -35,6 +33,4 @@ private:
 	TMap<int32, UE::Editor::DataStorage::RowHandle> EntryIndexToHandleMap;
 
 	UE::Editor::DataStorage::QueryHandle InitialColumnQuery;
-
-	//UE::Editor::DataStorage::QueryHandle CollectDataQuery;
 };
diff --git a/Engine/Plugins/Experimental/Animation/HierarchyTable/Source/Editor/Public/HierarchyTableLabelWidget.h b/Engine/Plugins/Experimental/Animation/HierarchyTable/Source/Editor/Public/HierarchyTableLabelWidget.h
new file mode 100644
--- /dev/null
+++ b/Engine/Plugins/Experimental/Animation/HierarchyTable/Source/Editor/Public/HierarchyTableLabelWidget.h
@@ -0,0 +1,43 @@
+// Copyright Epic Games, Inc. All Rights Reserved.
+
+#pragma once
+
+#include "Elements/Interfaces/TypedElementDataStorageFactory.h"
+#include "Elements/Interfaces/TypedElementDataStorageUiInterface.h"
+#include "UObject/ObjectMacros.h"
+
+#include "HierarchyTableLabelWidget.generated.h"
+
+class IEditorDataStorageProvider;
+class UScriptStruct;
+
+UCLASS()
+class UHierarchyTableLabelWidgetFactory : public UEditorDataStorageFactory
+{
+	GENERATED_BODY()
+
+public:
+	~UHierarchyTableLabelWidgetFactory() override = default;
+
+	void RegisterWidgetConstructors(
+		IEditorDataStorageProvider& DataStorage,
+		IEditorDataStorageUiProvider& DataStorageUi) const override;
+};
+
+USTRUCT()
+struct FHierarchyTableLabelWidgetConstructor : public FSimpleWidgetConstructor
+{
+	GENERATED_BODY()
+
+public:
+	FHierarchyTableLabelWidgetConstructor();
+	explicit FHierarchyTableLabelWidgetConstructor(const UScriptStruct* TypeInfo);
+	~FHierarchyTableLabelWidgetConstructor() override = default;
+
+	virtual TSharedPtr<SWidget> CreateWidget(
+		IEditorDataStorageProvider* DataStorage,
+		IEditorDataStorageUiProvider* DataStorageUi,
+		UE::Editor::DataStorage::RowHandle TargetRow,
+		UE::Editor::DataStorage::RowHandle WidgetRow, 
+		const UE::Editor::DataStorage::FMetaDataView& Arguments) override;
+};
\ No newline at end of file
