Entries added to hierarchy table are added as children of the selected row
* Also specify display name for Mask hierarchy table type

diff --git a/Engine/Plugins/Experimental/Animation/HierarchyTable/Source/Editor/Private/HierarchyTableEditorToolkit.cpp b/Engine/Plugins/Experimental/Animation/HierarchyTable/Source/Editor/Private/HierarchyTableEditorToolkit.cpp
--- a/Engine/Plugins/Experimental/Animation/HierarchyTable/Source/Editor/Private/HierarchyTableEditorToolkit.cpp
+++ b/Engine/Plugins/Experimental/Animation/HierarchyTable/Source/Editor/Private/HierarchyTableEditorToolkit.cpp
@@ -22,6 +22,7 @@
 #include "HierarchyTableEditorModule.h"
 #include "Widgets/Input/STextEntryPopup.h"
 #include "PersonaModule.h"
+#include "TedsOutlinerItem.h"
 
 #define LOCTEXT_NAMESPACE "HierarchyTableEditorToolkit"
 
@@ -193,6 +194,8 @@ TSharedRef<SWidget> FHierarchyTableEditorToolkit::CreateTedsOutliner()
 	}
 
 	TSharedRef<ISceneOutliner> TedsOutliner = TedsOutlinerModule.CreateTedsOutliner(InitOptions, Params, InitialColumnQuery);
+	TedsOutlinerPtr = TedsOutliner;
+
 	return TedsOutliner;
 }
 
@@ -321,20 +324,33 @@ void FHierarchyTableEditorToolkit::AddEntry(const FName Identifier, const EHiera
 		return;
 	}
 
+	using namespace UE::Editor::DataStorage;
+	UTypedElementRegistry* Registry = UTypedElementRegistry::GetInstance();
+	IEditorDataStorageProvider* DSI = Registry->GetMutableDataStorage();
+	static TableHandle Table = DSI->FindTable(FName("Editor_HierarchyTableTable"));
+
+	int32 ParentIndex = 0; // root of hierarchy
+
+	TArray<FSceneOutlinerTreeItemPtr> Selection = TedsOutlinerPtr->GetTree().GetSelectedItems();
+	if (Selection.Num() > 0)
+	{
+		const UE::Editor::Outliner::FTedsOutlinerTreeItem* TedsItem = Selection[0]->CastTo<UE::Editor::Outliner::FTedsOutlinerTreeItem>();
+		const UE::Editor::DataStorage::RowHandle ParentRowHandle = TedsItem->GetRowHandle();
+
+		const FTypedElementOverrideColumn* OverrideColumn = DSI->GetColumn<FTypedElementOverrideColumn>(ParentRowHandle);
+		ParentIndex = OverrideColumn->OwnerEntryIndex;
+	}
+
 	FHierarchyTableEntryData EntryData;
 	{
 		EntryData.Identifier = Identifier;
 		EntryData.EntryType = EntryType;
 		EntryData.OwnerTable = HierarchyTable;
-		EntryData.Parent = 0;
+		EntryData.Parent = ParentIndex;
 		EntryData.Payload = TOptional<FInstancedStruct>();
 	}
 	const int32 EntryIndex = HierarchyTable->TableData.Add(EntryData);
 
-	using namespace UE::Editor::DataStorage;
-	UTypedElementRegistry* Registry = UTypedElementRegistry::GetInstance();
-	IEditorDataStorageProvider* DSI = Registry->GetMutableDataStorage();
-	static TableHandle Table = DSI->FindTable(FName("Editor_HierarchyTableTable"));
 	const RowHandle Row = DSI->AddRow(Table);
 
 	{
diff --git a/Engine/Plugins/Experimental/Animation/HierarchyTable/Source/Editor/Public/HierarchyTableEditorToolkit.h b/Engine/Plugins/Experimental/Animation/HierarchyTable/Source/Editor/Public/HierarchyTableEditorToolkit.h
--- a/Engine/Plugins/Experimental/Animation/HierarchyTable/Source/Editor/Public/HierarchyTableEditorToolkit.h
+++ b/Engine/Plugins/Experimental/Animation/HierarchyTable/Source/Editor/Public/HierarchyTableEditorToolkit.h
@@ -33,4 +33,6 @@ private:
 	TMap<int32, UE::Editor::DataStorage::RowHandle> EntryIndexToHandleMap;
 
 	UE::Editor::DataStorage::QueryHandle InitialColumnQuery;
+
+	TSharedPtr<class ISceneOutliner> TedsOutlinerPtr;
 };
diff --git a/Engine/Plugins/Experimental/Animation/HierarchyTableBuiltin/Source/Editor/Public/MaskProfile/MaskProfileProxyColumn.h b/Engine/Plugins/Experimental/Animation/HierarchyTableBuiltin/Source/Editor/Public/MaskProfile/MaskProfileProxyColumn.h
--- a/Engine/Plugins/Experimental/Animation/HierarchyTableBuiltin/Source/Editor/Public/MaskProfile/MaskProfileProxyColumn.h
+++ b/Engine/Plugins/Experimental/Animation/HierarchyTableBuiltin/Source/Editor/Public/MaskProfile/MaskProfileProxyColumn.h
@@ -6,7 +6,7 @@
 
 #include "MaskProfileProxyColumn.generated.h"
 
-USTRUCT()
+USTRUCT(meta = (DisplayName = "Mask"))
 struct FHierarchyTableMaskColumn_Value final : public FEditorDataStorageColumn
 {
 	GENERATED_BODY()
