Update override column to be fixed width in hierarchy tables

diff --git a/Engine/Plugins/Experimental/Animation/HierarchyTable/Source/Editor/HierarchyTableEditor.Build.cs b/Engine/Plugins/Experimental/Animation/HierarchyTable/Source/Editor/HierarchyTableEditor.Build.cs
--- a/Engine/Plugins/Experimental/Animation/HierarchyTable/Source/Editor/HierarchyTableEditor.Build.cs
+++ b/Engine/Plugins/Experimental/Animation/HierarchyTable/Source/Editor/HierarchyTableEditor.Build.cs
@@ -22,7 +22,8 @@ public class HierarchyTableEditor : ModuleRules
 			"TypedElementFramework",
 			"SceneOutliner",
 			"ToolMenus",
-			"Persona"
+			"Persona",
+			"TedsTableViewer"
 		);
 	}
 }
diff --git a/Engine/Plugins/Experimental/Animation/HierarchyTable/Source/Editor/Private/HierarchyTable/Factories/OverrideFactory.cpp b/Engine/Plugins/Experimental/Animation/HierarchyTable/Source/Editor/Private/HierarchyTable/Factories/OverrideFactory.cpp
--- a/Engine/Plugins/Experimental/Animation/HierarchyTable/Source/Editor/Private/HierarchyTable/Factories/OverrideFactory.cpp
+++ b/Engine/Plugins/Experimental/Animation/HierarchyTable/Source/Editor/Private/HierarchyTable/Factories/OverrideFactory.cpp
@@ -9,5 +9,9 @@ void UHierarchyTableOverrideFactory::RegisterWidgetConstructors(IEditorDataStora
 {
 	using namespace UE::Editor::DataStorage::Queries;
 
-	DataStorageUi.RegisterWidgetFactory<FTypedElementWidgetConstructor_Override>(FName(TEXT("General.Cell")), TColumn<FTypedElementOverrideColumn>());
+	DataStorageUi.RegisterWidgetFactory<FTypedElementWidgetConstructor_Override>(FName(TEXT("General.Cell")), 
+		TColumn<FTypedElementOverrideColumn>());
+
+	DataStorageUi.RegisterWidgetFactory<FTypedElementWidgetHeaderConstructor_Override>(FName(TEXT("General.Header")),
+		TColumn<FTypedElementOverrideColumn>());
 }
\ No newline at end of file
diff --git a/Engine/Plugins/Experimental/Animation/HierarchyTable/Source/Editor/Private/HierarchyTable/Widgets/OverrideWidgetConstructor.cpp b/Engine/Plugins/Experimental/Animation/HierarchyTable/Source/Editor/Private/HierarchyTable/Widgets/OverrideWidgetConstructor.cpp
--- a/Engine/Plugins/Experimental/Animation/HierarchyTable/Source/Editor/Private/HierarchyTable/Widgets/OverrideWidgetConstructor.cpp
+++ b/Engine/Plugins/Experimental/Animation/HierarchyTable/Source/Editor/Private/HierarchyTable/Widgets/OverrideWidgetConstructor.cpp
@@ -8,6 +8,9 @@
 #include "Widgets/Layout/SBox.h"
 #include "Widgets/SBoxPanel.h"
 #include "HierarchyTable.h"
+#include "Columns/UIPropertiesColumns.h"
+
+#define LOCTEXT_NAMESPACE "FTypedElementWidgetConstructor_Override"
 
 FTypedElementWidgetConstructor_Override::FTypedElementWidgetConstructor_Override()
 	: Super(StaticStruct())
@@ -107,4 +110,36 @@ bool FTypedElementWidgetConstructor_Override::FinalizeWidget(
 	WidgetInstance->SetContent(NewWidget);
 
 	return true;
-}
\ No newline at end of file
+}
+
+//
+// FTypedElementWidgetHeaderConstructor_Override
+//
+
+FTypedElementWidgetHeaderConstructor_Override::FTypedElementWidgetHeaderConstructor_Override()
+	: Super(StaticStruct())
+{
+}
+
+TSharedPtr<SWidget> FTypedElementWidgetHeaderConstructor_Override::CreateWidget(
+	const UE::Editor::DataStorage::FMetaDataView& Arguments)
+{
+	return SNew(SImage)
+		.DesiredSizeOverride(FVector2D(16.f, 16.f))
+		.ColorAndOpacity(FSlateColor::UseForeground())
+		.Image(FAppStyle::GetBrush("DetailsView.OverrideHere"))
+		.ToolTipText(FText(LOCTEXT("OverrideColumnHeader", "Overrides")));
+}
+
+bool FTypedElementWidgetHeaderConstructor_Override::FinalizeWidget(IEditorDataStorageProvider* DataStorage, 
+	IEditorDataStorageUiProvider* DataStorageUi, UE::Editor::DataStorage::RowHandle Row, const TSharedPtr<SWidget>& Widget)
+{
+	DataStorage->AddColumn(Row, FUIHeaderPropertiesColumn
+		{
+			.ColumnSizeMode = EColumnSizeMode::Fixed,
+			.Width = 24.0f
+		});
+	return true;
+}
+
+#undef LOCTEXT_NAMESPACE
\ No newline at end of file
diff --git a/Engine/Plugins/Experimental/Animation/HierarchyTable/Source/Editor/Private/HierarchyTable/Widgets/OverrideWidgetConstructor.h b/Engine/Plugins/Experimental/Animation/HierarchyTable/Source/Editor/Private/HierarchyTable/Widgets/OverrideWidgetConstructor.h
--- a/Engine/Plugins/Experimental/Animation/HierarchyTable/Source/Editor/Private/HierarchyTable/Widgets/OverrideWidgetConstructor.h
+++ b/Engine/Plugins/Experimental/Animation/HierarchyTable/Source/Editor/Private/HierarchyTable/Widgets/OverrideWidgetConstructor.h
@@ -19,4 +19,19 @@ protected:
 	TSharedPtr<SWidget> CreateWidget(const UE::Editor::DataStorage::FMetaDataView& Arguments);
 	bool FinalizeWidget(IEditorDataStorageProvider* DataStorage, IEditorDataStorageUiProvider* DataStorageUi,
 		UE::Editor::DataStorage::RowHandle Row, const TSharedPtr<SWidget>& Widget) override;
+};
+
+USTRUCT()
+struct FTypedElementWidgetHeaderConstructor_Override : public FTypedElementWidgetConstructor
+{
+	GENERATED_BODY()
+
+public:
+	FTypedElementWidgetHeaderConstructor_Override();
+	~FTypedElementWidgetHeaderConstructor_Override() override = default;
+
+protected:
+	TSharedPtr<SWidget> CreateWidget(const UE::Editor::DataStorage::FMetaDataView& Arguments) override;
+	bool FinalizeWidget(IEditorDataStorageProvider* DataStorage, IEditorDataStorageUiProvider* DataStorageUi,
+		UE::Editor::DataStorage::RowHandle Row, const TSharedPtr<SWidget>& Widget) override;
 };
\ No newline at end of file
