Remove unused hierarchy table functions

diff --git a/Engine/Plugins/Experimental/Animation/HierarchyTable/Source/Editor/Private/DefaultHierarchyTableTypeHandler.cpp b/Engine/Plugins/Experimental/Animation/HierarchyTable/Source/Editor/Private/DefaultHierarchyTableTypeHandler.cpp
--- a/Engine/Plugins/Experimental/Animation/HierarchyTable/Source/Editor/Private/DefaultHierarchyTableTypeHandler.cpp
+++ b/Engine/Plugins/Experimental/Animation/HierarchyTable/Source/Editor/Private/DefaultHierarchyTableTypeHandler.cpp
@@ -8,11 +8,6 @@
 
 #define LOCTEXT_NAMESPACE "UHierarchyTable_TableTypeHandler_Default"
 
-TSharedPtr<FHierarchyTable_TableType> UHierarchyTable_TableTypeHandler_Default::CreateTableType() const
-{
-	return MakeShared<FHierarchyTable_TableType_Default>();
-}
-
 void UHierarchyTable_TableTypeHandler_Default::InitialiseTable(const FInstancedStruct& TableType, TObjectPtr<UHierarchyTable> HierarchyTable) const
 {
 	// Do nothing, hierarchy is built manually
diff --git a/Engine/Plugins/Experimental/Animation/HierarchyTable/Source/Editor/Private/DefaultHierarchyTableTypeHandler.h b/Engine/Plugins/Experimental/Animation/HierarchyTable/Source/Editor/Private/DefaultHierarchyTableTypeHandler.h
--- a/Engine/Plugins/Experimental/Animation/HierarchyTable/Source/Editor/Private/DefaultHierarchyTableTypeHandler.h
+++ b/Engine/Plugins/Experimental/Animation/HierarchyTable/Source/Editor/Private/DefaultHierarchyTableTypeHandler.h
@@ -12,8 +12,6 @@ class UHierarchyTable_TableTypeHandler_Default final : public UHierarchyTable_Ta
 	GENERATED_BODY()
 
 public:
-	virtual TSharedPtr<FHierarchyTable_TableType> CreateTableType() const;
-
 	virtual void InitialiseTable(const FInstancedStruct& TableType, TObjectPtr<UHierarchyTable> HierarchyTable) const;
 
 	virtual bool FactoryConfigureProperties(FInstancedStruct& TableType) const override;
diff --git a/Engine/Plugins/Experimental/Animation/HierarchyTable/Source/Editor/Public/HierarchyTableTypeHandler.h b/Engine/Plugins/Experimental/Animation/HierarchyTable/Source/Editor/Public/HierarchyTableTypeHandler.h
--- a/Engine/Plugins/Experimental/Animation/HierarchyTable/Source/Editor/Public/HierarchyTableTypeHandler.h
+++ b/Engine/Plugins/Experimental/Animation/HierarchyTable/Source/Editor/Public/HierarchyTableTypeHandler.h
@@ -19,14 +19,10 @@ class HIERARCHYTABLEEDITOR_API UHierarchyTable_TableTypeHandler : public UObject
 	GENERATED_BODY()
 
 public:
-	virtual TSharedPtr<FHierarchyTable_TableType> CreateTableType() const { return nullptr; }
-
 	virtual void InitialiseTable(const FInstancedStruct& TableType, TObjectPtr<UHierarchyTable> HierarchyTable) const {}
 
 	virtual bool FactoryConfigureProperties(FInstancedStruct& TableType) const { return false; }
 
-	//virtual int32 AddEntry(TObjectPtr<UHierarchyTable> HierarchyTable, const int32 ParentIndex) const { return false; }
-
 	virtual void ExtendToolbar(UToolMenu* ToolMenu, UHierarchyTable* HierarchyTable, const IHierarchyTableEditor* Editor) const {}
 
 	virtual TSharedRef<SWidget> CreateLabelWidget(TObjectPtr<UHierarchyTable> HierarchyTable, const int32 EntryIndex) const;
diff --git a/Engine/Plugins/Experimental/Animation/HierarchyTableBuiltin/Source/Editor/Private/SkeletonHierarchyTableTypeHandler.cpp b/Engine/Plugins/Experimental/Animation/HierarchyTableBuiltin/Source/Editor/Private/SkeletonHierarchyTableTypeHandler.cpp
--- a/Engine/Plugins/Experimental/Animation/HierarchyTableBuiltin/Source/Editor/Private/SkeletonHierarchyTableTypeHandler.cpp
+++ b/Engine/Plugins/Experimental/Animation/HierarchyTableBuiltin/Source/Editor/Private/SkeletonHierarchyTableTypeHandler.cpp
@@ -24,11 +24,6 @@
 
 #define LOCTEXT_NAMESPACE "UHierarchyTable_TableTypeHandler_Skeleton"
 
-TSharedPtr<FHierarchyTable_TableType> UHierarchyTable_TableTypeHandler_Skeleton::CreateTableType() const
-{
-	return MakeShared<FHierarchyTable_TableType_Skeleton>();
-}
-
 void UHierarchyTable_TableTypeHandler_Skeleton::InitialiseTable(const FInstancedStruct& TableType, TObjectPtr<UHierarchyTable> HierarchyTable) const
 {
 	const FHierarchyTable_TableType_Skeleton& SkeletonTableType = TableType.Get<FHierarchyTable_TableType_Skeleton>();
diff --git a/Engine/Plugins/Experimental/Animation/HierarchyTableBuiltin/Source/Editor/Private/SkeletonHierarchyTableTypeHandler.h b/Engine/Plugins/Experimental/Animation/HierarchyTableBuiltin/Source/Editor/Private/SkeletonHierarchyTableTypeHandler.h
--- a/Engine/Plugins/Experimental/Animation/HierarchyTableBuiltin/Source/Editor/Private/SkeletonHierarchyTableTypeHandler.h
+++ b/Engine/Plugins/Experimental/Animation/HierarchyTableBuiltin/Source/Editor/Private/SkeletonHierarchyTableTypeHandler.h
@@ -12,8 +12,6 @@ class UHierarchyTable_TableTypeHandler_Skeleton final : public UHierarchyTable_T
 	GENERATED_BODY()
 
 public:
-	virtual TSharedPtr<FHierarchyTable_TableType> CreateTableType() const;
-
 	virtual void InitialiseTable(const FInstancedStruct& TableType, TObjectPtr<UHierarchyTable> HierarchyTable) const;
 
 	virtual bool FactoryConfigureProperties(FInstancedStruct& TableType) const override;
