Hierarchy Tables - Update row style to have alternating rows

diff --git a/Engine/Plugins/Experimental/Animation/HierarchyTable/Source/Editor/Private/SHierarchyTableRow.cpp b/Engine/Plugins/Experimental/Animation/HierarchyTable/Source/Editor/Private/SHierarchyTableRow.cpp
--- a/Engine/Plugins/Experimental/Animation/HierarchyTable/Source/Editor/Private/SHierarchyTableRow.cpp
+++ b/Engine/Plugins/Experimental/Animation/HierarchyTable/Source/Editor/Private/SHierarchyTableRow.cpp
@@ -39,6 +39,7 @@ void SHierarchyTableRow::Construct(const FArguments& InArgs, const TSharedRef<ST
 
 	FSuperRowType::Construct(
 		FSuperRowType::FArguments()
+		.Style( FAppStyle::Get(), "TableView.AlternatingRow" )
 		.OnDragDetected(this, &SHierarchyTableRow::OnDragDetected)
 		.OnCanAcceptDrop(this, &SHierarchyTableRow::OnCanAcceptDrop)
 		.OnAcceptDrop(this, &SHierarchyTableRow::OnAcceptDrop),
