Add float type for hierarchy tables

diff --git a/Engine/Plugins/Experimental/Animation/HierarchyTable/Source/Editor/Private/FloatColumn.cpp b/Engine/Plugins/Experimental/Animation/HierarchyTable/Source/Editor/Private/FloatColumn.cpp
new file mode 100644
--- /dev/null
+++ b/Engine/Plugins/Experimental/Animation/HierarchyTable/Source/Editor/Private/FloatColumn.cpp
@@ -0,0 +1,41 @@
+// Copyright Epic Games, Inc. All Rights Reserved.
+
+#include "FloatColumn.h"
+
+#include "Widgets/Input/SSpinBox.h"
+#include "HierarchyTable.h"
+#include "Editor.h"
+#include "HierarchyTableDefaultTypes.h"
+
+#define LOCTEXT_NAMESPACE "FHierarchyTableColumn_Float"
+
+TSharedRef<SWidget> FHierarchyTableColumn_Float::CreateHeaderWidget()
+{
+	return SNullWidget::NullWidget;
+}
+
+TSharedRef<SWidget> FHierarchyTableColumn_Float::CreateEntryWidget(TObjectPtr<UHierarchyTable> HierarchyTable, int32 EntryIndex)
+{
+	return SNew(SSpinBox<float>)
+		.IsEnabled_Lambda([HierarchyTable, EntryIndex]() { return HierarchyTable->GetTableEntry(EntryIndex)->IsOverridden(); })
+		.MinDesiredWidth(100.0f)
+		.Value_Lambda([HierarchyTable, EntryIndex]()
+			{
+				return HierarchyTable->GetTableEntry(EntryIndex)->GetValue<FHierarchyTable_ElementType_Float>()->Value;
+			})
+		.OnValueChanged_Lambda([HierarchyTable, EntryIndex](float NewValue)
+			{
+				HierarchyTable->GetMutableTableEntry(EntryIndex)->GetMutableValue<FHierarchyTable_ElementType_Float>()->Value = NewValue;
+			})
+		.OnBeginSliderMovement_Lambda([HierarchyTable]()
+			{
+				GEditor->BeginTransaction(LOCTEXT("SetFloatValue", "Set Float Value"));
+				HierarchyTable->Modify();
+			})
+		.OnEndSliderMovement_Lambda([](float)
+			{
+				GEditor->EndTransaction();
+			});
+}
+
+#undef LOCTEXT_NAMESPACE
\ No newline at end of file
diff --git a/Engine/Plugins/Experimental/Animation/HierarchyTable/Source/Editor/Private/FloatColumn.h b/Engine/Plugins/Experimental/Animation/HierarchyTable/Source/Editor/Private/FloatColumn.h
new file mode 100644
--- /dev/null
+++ b/Engine/Plugins/Experimental/Animation/HierarchyTable/Source/Editor/Private/FloatColumn.h
@@ -0,0 +1,20 @@
+// Copyright Epic Games, Inc. All Rights Reserved.
+
+#pragma once
+
+#include "IHierarchyTableColumn.h"
+
+#define LOCTEXT_NAMESPACE "HierarchyTableColumn_Float"
+
+struct FHierarchyTableColumn_Float : public IHierarchyTableColumn
+{
+public:
+	virtual FName GetColumnId() const override { return FName("Float"); }
+	virtual FText GetColumnLabel() const override { return LOCTEXT("FloatLabel", "Float"); }
+	virtual float GetColumnSize() const override { return 1.0f; }
+
+	virtual TSharedRef<SWidget> CreateEntryWidget(TObjectPtr<UHierarchyTable> HierarchyTable, int32 EntryIndex);
+	virtual TSharedRef<SWidget> CreateHeaderWidget();
+};
+
+#undef LOCTEXT_NAMESPACE
\ No newline at end of file
diff --git a/Engine/Plugins/Experimental/Animation/HierarchyTable/Source/Editor/Private/HierachyTableEditorModule.cpp b/Engine/Plugins/Experimental/Animation/HierarchyTable/Source/Editor/Private/HierachyTableEditorModule.cpp
--- a/Engine/Plugins/Experimental/Animation/HierarchyTable/Source/Editor/Private/HierachyTableEditorModule.cpp
+++ b/Engine/Plugins/Experimental/Animation/HierarchyTable/Source/Editor/Private/HierachyTableEditorModule.cpp
@@ -3,12 +3,20 @@
 #include "HierarchyTableEditorModule.h"
 #include "Modules/ModuleManager.h"
 #include "DefaultHierarchyTableType.h"
+#include "HierarchyTableDefaultTypes.h"
 #include "DefaultHierarchyTableTypeHandler.h"
+#include "FloatColumn.h"
 
 void FHierarchyTableEditorModule::StartupModule()
 {
 	BuiltinTableTypes.Add(FHierarchyTable_TableType_Default::StaticStruct());
 	RegisterTableType(FHierarchyTable_TableType_Default::StaticStruct(), UHierarchyTable_TableTypeHandler_Default::StaticClass());
+
+	BuiltinElementTypes.Add(FHierarchyTable_ElementType_Float::StaticStruct());
+	RegisterElementTypeEditorColumns(FHierarchyTable_ElementType_Float::StaticStruct(),
+		{	
+			MakeShared<FHierarchyTableColumn_Float>(),
+		});
 }
 
 void FHierarchyTableEditorModule::ShutdownModule()
@@ -20,6 +28,14 @@ void FHierarchyTableEditorModule::ShutdownModule()
 			UnregisterTableType(StructPtr);
 		}
 	}
+
+	for (TWeakObjectPtr<UScriptStruct> WeakPtr : BuiltinElementTypes)
+	{
+		if (const UScriptStruct* StructPtr = WeakPtr.Get())
+		{
+			UnregisterElementTypeEditorColumns(StructPtr);
+		}
+	}
 }
 
 void FHierarchyTableEditorModule::RegisterTableType(const UScriptStruct* TableType, const UClass* Handler)
diff --git a/Engine/Plugins/Experimental/Animation/HierarchyTable/Source/Editor/Public/HierarchyTableEditorModule.h b/Engine/Plugins/Experimental/Animation/HierarchyTable/Source/Editor/Public/HierarchyTableEditorModule.h
--- a/Engine/Plugins/Experimental/Animation/HierarchyTable/Source/Editor/Public/HierarchyTableEditorModule.h
+++ b/Engine/Plugins/Experimental/Animation/HierarchyTable/Source/Editor/Public/HierarchyTableEditorModule.h
@@ -34,6 +34,8 @@ public:
 private:
 	TArray<TWeakObjectPtr<UScriptStruct>> BuiltinTableTypes;
 
+	TArray<TWeakObjectPtr<UScriptStruct>> BuiltinElementTypes;
+
 	TMap<TObjectKey<UScriptStruct>, const UClass*> TableHandlers;
 
 	TMap<TObjectKey<UScriptStruct>, TArray<TSharedPtr<IHierarchyTableColumn>>> EditorColumns;
diff --git a/Engine/Plugins/Experimental/Animation/HierarchyTable/Source/Runtime/Public/HierarchyTableDefaultTypes.h b/Engine/Plugins/Experimental/Animation/HierarchyTable/Source/Runtime/Public/HierarchyTableDefaultTypes.h
new file mode 100644
--- /dev/null
+++ b/Engine/Plugins/Experimental/Animation/HierarchyTable/Source/Runtime/Public/HierarchyTableDefaultTypes.h
@@ -0,0 +1,21 @@
+// Copyright Epic Games, Inc. All Rights Reserved.
+
+#pragma once
+
+#include "HierarchyTableType.h"
+
+#include "HierarchyTableDefaultTypes.generated.h"
+
+USTRUCT(DisplayName = "Float")
+struct FHierarchyTable_ElementType_Float final : public FHierarchyTable_ElementType
+{
+	GENERATED_BODY()
+
+	FHierarchyTable_ElementType_Float()
+		: Value(0.0)
+	{
+	}
+
+	UPROPERTY()
+	float Value;
+};
\ No newline at end of file
