Fix hierarchy table editor crash due to GC

diff --git a/Engine/Plugins/Experimental/Animation/HierarchyTable/Source/Editor/Private/SHierarchyTable.cpp b/Engine/Plugins/Experimental/Animation/HierarchyTable/Source/Editor/Private/SHierarchyTable.cpp
--- a/Engine/Plugins/Experimental/Animation/HierarchyTable/Source/Editor/Private/SHierarchyTable.cpp
+++ b/Engine/Plugins/Experimental/Animation/HierarchyTable/Source/Editor/Private/SHierarchyTable.cpp
@@ -275,4 +275,16 @@ void SHierarchyTable::PostRedo(bool bSuccess)
 	RegenerateTreeViewItems();
 }
 
+void SHierarchyTable::AddReferencedObjects(FReferenceCollector& Collector)
+{
+	Collector.AddReferencedObject(HierarchyTable);
+	Collector.AddReferencedObject(TableHandler);
+}
+
+
+FString SHierarchyTable::GetReferencerName() const
+{
+	return FString(TEXT("SHierarchyTable"));
+}
+
 #undef LOCTEXT_NAMESPACE
\ No newline at end of file
diff --git a/Engine/Plugins/Experimental/Animation/HierarchyTable/Source/Editor/Private/SHierarchyTable.h b/Engine/Plugins/Experimental/Animation/HierarchyTable/Source/Editor/Private/SHierarchyTable.h
--- a/Engine/Plugins/Experimental/Animation/HierarchyTable/Source/Editor/Private/SHierarchyTable.h
+++ b/Engine/Plugins/Experimental/Animation/HierarchyTable/Source/Editor/Private/SHierarchyTable.h
@@ -14,7 +14,7 @@ class SHierarchyTableRow;
 class UHierarchyTable;
 class UHierarchyTable_TableTypeHandler;
 
-class SHierarchyTable : public SCompoundWidget, public IHierarchyTable, public FSelfRegisteringEditorUndoClient
+class SHierarchyTable : public SCompoundWidget, public IHierarchyTable, public FSelfRegisteringEditorUndoClient, public FGCObject
 {
 public:
 	friend SHierarchyTableRow;
@@ -48,6 +48,10 @@ public:
 	void PostUndo(bool bSuccess) override;
 	void PostRedo(bool bSuccess) override;
 
+	// FGCObject
+	virtual void AddReferencedObjects(FReferenceCollector& Collector) override;
+	virtual FString GetReferencerName() const override;
+
 private:
 	void TreeView_HandleGetChildren(TSharedPtr<FTreeItem> InItem, TArray<TSharedPtr<FTreeItem>>& OutChildren) const;
 
