Move hierarchy table from animation category to misc.

diff --git a/Engine/Plugins/Experimental/Animation/HierarchyTable/Source/Editor/Private/HierarchyTableAssetDefinition.cpp b/Engine/Plugins/Experimental/Animation/HierarchyTable/Source/Editor/Private/HierarchyTableAssetDefinition.cpp
--- a/Engine/Plugins/Experimental/Animation/HierarchyTable/Source/Editor/Private/HierarchyTableAssetDefinition.cpp
+++ b/Engine/Plugins/Experimental/Animation/HierarchyTable/Source/Editor/Private/HierarchyTableAssetDefinition.cpp
@@ -36,7 +36,7 @@ EAssetCommandResult UAssetDefinition_HierarchyTable::OpenAssets(const FAssetOpen
 
 TConstArrayView<FAssetCategoryPath> UAssetDefinition_HierarchyTable::GetAssetCategories() const
 {
-	static const auto Categories = { FAssetCategoryPath(EAssetCategoryPaths::Animation) };
+	static const auto Categories = { FAssetCategoryPath(EAssetCategoryPaths::Misc) };
 	return Categories;
 }
 
